Add ability to define expected content type and content formatting for input on a RPC to be shown on ListRPCs

diff --git a/Engine/Source/Runtime/ExternalRPCRegistry/Private/ExternalRpcRegistry.cpp b/Engine/Source/Runtime/ExternalRPCRegistry/Private/ExternalRpcRegistry.cpp
--- a/Engine/Source/Runtime/ExternalRPCRegistry/Private/ExternalRpcRegistry.cpp
+++ b/Engine/Source/Runtime/ExternalRPCRegistry/Private/ExternalRpcRegistry.cpp
@@ -83,20 +83,24 @@ bool UExternalRpcRegistry::GetRegisteredRoute(FName RouteName, FExternalRouteInf
 	if (RegisteredRoutes.Find(RouteName))
 	{
 		OutRouteInfo.RouteName = RouteName;
-		OutRouteInfo.RoutePath = RegisteredRoutes[RouteName]->Path;
-		OutRouteInfo.RequestVerbs = RegisteredRoutes[RouteName]->Verbs;
+		OutRouteInfo.RoutePath = RegisteredRoutes[RouteName].Handle->Path;
+		OutRouteInfo.RequestVerbs = RegisteredRoutes[RouteName].Handle->Verbs;
+		OutRouteInfo.InputContentType = RegisteredRoutes[RouteName].InputContentType;
+		OutRouteInfo.InputExpectedFormat = RegisteredRoutes[RouteName].InputExpectedFormat;
 		return true;
 	}
 	return false;
 }
 
-void UExternalRpcRegistry::RegisterNewRoute(FName RouteName, const FHttpPath& HttpPath, const EHttpServerRequestVerbs& RequestVerbs, const FHttpRequestHandler& Handler, bool bOverrideIfBound /* = false */)
+void UExternalRpcRegistry::RegisterNewRoute(FName RouteName, const FHttpPath& HttpPath, const EHttpServerRequestVerbs& RequestVerbs, const FHttpRequestHandler& Handler, bool bOverrideIfBound /* = false */, FString OptionalContentType /* = TEXT("")*/, FString OptionalExpectedFormat /*= TEXT("")*/)
 {
 #if WITH_RPC_REGISTRY
 	FExternalRouteInfo InRouteInfo;
 	InRouteInfo.RouteName = RouteName;
 	InRouteInfo.RoutePath = HttpPath;
 	InRouteInfo.RequestVerbs = RequestVerbs;
+	InRouteInfo.InputContentType = OptionalContentType;
+	InRouteInfo.InputExpectedFormat = OptionalExpectedFormat;
 	RegisterNewRoute(InRouteInfo, Handler, bOverrideIfBound);
 #endif
 }
@@ -112,10 +116,14 @@ void UExternalRpcRegistry::RegisterNewRoute(FExternalRouteInfo InRouteInfo, cons
 			UE_LOG(LogExternalRpcRegistry, Error, TEXT("Failed to bind route with friendly key %s - a route at location %s already exists."), *InRouteInfo.RouteName.ToString(), *InRouteInfo.RoutePath.GetPath());
 			return;
 		}
-		UE_LOG(LogExternalRpcRegistry, Log, TEXT("Overwriting route at friendly key %s - from %s to %s "), *InRouteInfo.RouteName.ToString(), *RegisteredRoutes[InRouteInfo.RouteName]->Path, *InRouteInfo.RoutePath.GetPath());
-		HttpRouter->UnbindRoute(RegisteredRoutes[InRouteInfo.RouteName]);
+		UE_LOG(LogExternalRpcRegistry, Log, TEXT("Overwriting route at friendly key %s - from %s to %s "), *InRouteInfo.RouteName.ToString(), *RegisteredRoutes[InRouteInfo.RouteName].Handle->Path, *InRouteInfo.RoutePath.GetPath());
+		HttpRouter->UnbindRoute(RegisteredRoutes[InRouteInfo.RouteName].Handle);
 	}
-	RegisteredRoutes.Add(InRouteInfo.RouteName, HttpRouter->BindRoute(InRouteInfo.RoutePath, InRouteInfo.RequestVerbs, Handler));
+	FExternalRouteDesc RouteDesc;
+	RouteDesc.Handle = HttpRouter->BindRoute(InRouteInfo.RoutePath, InRouteInfo.RequestVerbs, Handler);
+	RouteDesc.InputContentType = InRouteInfo.InputContentType;
+	RouteDesc.InputExpectedFormat = InRouteInfo.InputExpectedFormat;
+	RegisteredRoutes.Add(InRouteInfo.RouteName, RouteDesc);
 #endif
 }
 
@@ -126,7 +134,7 @@ void UExternalRpcRegistry::CleanUpRoute(FName RouteName, bool bFailIfUnbound /*
 	if (RegisteredRoutes.Find(RouteName))
 	{
 		TSharedPtr<IHttpRouter> HttpRouter = FHttpServerModule::Get().GetHttpRouter(PortToUse);
-		HttpRouter->UnbindRoute(RegisteredRoutes[RouteName]);
+		HttpRouter->UnbindRoute(RegisteredRoutes[RouteName].Handle);
 	}
 	else
 	{
@@ -148,8 +156,16 @@ bool UExternalRpcRegistry::HttpListOpenRoutes(const FHttpServerRequest& Request,
 	{
 		JsonWriter->WriteObjectStart();
 		JsonWriter->WriteValue(TEXT("name"), RouteKey.ToString());
-		JsonWriter->WriteValue(TEXT("route"), RegisteredRoutes[RouteKey]->Path);
-		JsonWriter->WriteValue(TEXT("verb"), GetHttpRouteVerbString(RegisteredRoutes[RouteKey]->Verbs));
+		JsonWriter->WriteValue(TEXT("route"), RegisteredRoutes[RouteKey].Handle->Path);
+		JsonWriter->WriteValue(TEXT("verb"), GetHttpRouteVerbString(RegisteredRoutes[RouteKey].Handle->Verbs));
+		if (!RegisteredRoutes[RouteKey].InputContentType.IsEmpty())
+		{
+			JsonWriter->WriteValue(TEXT("inputContentType"), RegisteredRoutes[RouteKey].InputContentType);
+		}
+		if (!RegisteredRoutes[RouteKey].InputExpectedFormat.IsEmpty())
+		{
+			JsonWriter->WriteValue(TEXT("inputExpectedFormat"), RegisteredRoutes[RouteKey].InputExpectedFormat);
+		}
 		JsonWriter->WriteObjectEnd();
 	}
 	JsonWriter->WriteArrayEnd();
diff --git a/Engine/Source/Runtime/ExternalRPCRegistry/Public/ExternalRpcRegistry.h b/Engine/Source/Runtime/ExternalRPCRegistry/Public/ExternalRpcRegistry.h
--- a/Engine/Source/Runtime/ExternalRPCRegistry/Public/ExternalRpcRegistry.h
+++ b/Engine/Source/Runtime/ExternalRPCRegistry/Public/ExternalRpcRegistry.h
@@ -17,20 +17,45 @@ public:
 	FName RouteName;
 	FHttpPath RoutePath;
 	EHttpServerRequestVerbs RequestVerbs;
+	FString InputContentType;
+	FString InputExpectedFormat;
 	FExternalRouteInfo()
 	{
 		RouteName = FName(TEXT(""));
 		RoutePath = FHttpPath();
 		RequestVerbs = EHttpServerRequestVerbs::VERB_NONE;
+		InputContentType = TEXT("");
+		InputExpectedFormat = TEXT("");
 	}
-	FExternalRouteInfo(FName InRouteName, FHttpPath InRoutePath, EHttpServerRequestVerbs InRequestVerbs)
+	FExternalRouteInfo(FName InRouteName, FHttpPath InRoutePath, EHttpServerRequestVerbs InRequestVerbs, FString InContentType = TEXT(""), FString InExpectedFormat = TEXT(""))
 	{
 		RouteName = InRouteName;
 		RoutePath = InRoutePath;
 		RequestVerbs = InRequestVerbs;
+		InputContentType = InContentType;
+		InputExpectedFormat = InExpectedFormat;
 	}
 };
 
+USTRUCT()
+struct FExternalRouteDesc
+{
+	GENERATED_BODY()
+public:
+	FHttpRouteHandle Handle;
+	FString InputContentType;
+	FString InputExpectedFormat;
+	FExternalRouteDesc()
+	{
+	
+	}
+	FExternalRouteDesc(FHttpRouteHandle InHandle, FString InContentType, FString InExpectedFormat)
+	{
+		Handle = InHandle;
+		InputContentType = InContentType;
+		InputExpectedFormat = InExpectedFormat;
+	}
+};
 
 /**
  * This class is designed to be a singleton that handles registry, maintenance, and cleanup of any REST endpoints exposed on the process 
@@ -42,7 +67,7 @@ class EXTERNALRPCREGISTRY_API UExternalRpcRegistry : public UObject
 	GENERATED_BODY()
 protected:
 	static UExternalRpcRegistry * ObjectInstance;
-	TMap<FName, FHttpRouteHandle> RegisteredRoutes;
+	TMap<FName, FExternalRouteDesc> RegisteredRoutes;
 
 public:
 	static UExternalRpcRegistry * GetInstance();
@@ -60,7 +85,7 @@ public:
 	 * Register a new route.
 	 * Will override existing routes if option is set, otherwise will error and fail to bind.
 	 */
-	void RegisterNewRoute(FName RouteName, const FHttpPath& HttpPath, const EHttpServerRequestVerbs& RequestVerbs, const FHttpRequestHandler& Handler, bool bOverrideIfBound = false);
+	void RegisterNewRoute(FName RouteName, const FHttpPath& HttpPath, const EHttpServerRequestVerbs& RequestVerbs, const FHttpRequestHandler& Handler, bool bOverrideIfBound = false, FString OptionalContentType = TEXT(""), FString OptionalExpectedFormat = TEXT(""));
 
 	/**
 	 * Clean up a route.
