Initial version of external RPC registry object for engine. Note: still working on making cleaner safeguards to prevent people from paving each other's routes in code.

diff --git a/Engine/Source/Runtime/ExternalRPCRegistry/ExternalRpcRegistry.Build.cs b/Engine/Source/Runtime/ExternalRPCRegistry/ExternalRpcRegistry.Build.cs
new file mode 100644
--- /dev/null
+++ b/Engine/Source/Runtime/ExternalRPCRegistry/ExternalRpcRegistry.Build.cs
@@ -0,0 +1,20 @@
+// Copyright 1998-2019 Epic Games, Inc. All Rights Reserved.
+
+using UnrealBuildTool;
+
+public class ExternalRpcRegistry : ModuleRules
+{
+    public ExternalRpcRegistry(ReadOnlyTargetRules Target) : base(Target)
+    {
+        PrivateDependencyModuleNames.AddRange(
+            new string[] {
+                "Core",
+                "CoreUObject",
+                "InputCore",
+				"Json",
+				"JsonUtilities",
+                "HTTPServer"
+            }
+        );
+    }
+}
diff --git a/Engine/Source/Runtime/ExternalRPCRegistry/Private/ExternalRpcRegistry.cpp b/Engine/Source/Runtime/ExternalRPCRegistry/Private/ExternalRpcRegistry.cpp
new file mode 100644
--- /dev/null
+++ b/Engine/Source/Runtime/ExternalRPCRegistry/Private/ExternalRpcRegistry.cpp
@@ -0,0 +1,159 @@
+// Copyright 1998-2019 Epic Games, Inc. All Rights Reserved.
+
+#include "ExternalRpcRegistry.h"
+#include "HttpServerModule.h"
+#include "Misc/CommandLine.h"
+#include "Misc/Parse.h"
+#include "HttpServerResponse.h"
+#include "Serialization/JsonWriter.h"
+#include "Logging/LogMacros.h"
+
+#define USE_RPC_REGISTRY_IN_SHIPPING 0
+#define USE_RPC_REGISTRY (USE_RPC_REGISTRY_IN_SHIPPING || !UE_BUILD_SHIPPING )
+
+DEFINE_LOG_CATEGORY(LogExternalRpcRegistry)
+
+UExternalRpcRegistry* UExternalRpcRegistry::ObjectInstance = nullptr;
+
+
+FString GetHttpRouteVerbString(EHttpServerRequestVerbs InVerbs)
+{
+#if USE_RPC_REGISTRY
+	switch (InVerbs)
+	{
+		case EHttpServerRequestVerbs::VERB_POST:
+		{
+			return TEXT("POST");
+		}
+		case EHttpServerRequestVerbs::VERB_PUT:
+		{
+			return TEXT("PUT");
+		}
+		case EHttpServerRequestVerbs::VERB_GET:
+		{
+			return TEXT("GET");
+		}
+		case EHttpServerRequestVerbs::VERB_PATCH:
+		{
+			return TEXT("PATCH");
+		}
+		case EHttpServerRequestVerbs::VERB_DELETE:
+		{
+			return TEXT("DELETE");
+		}
+		case EHttpServerRequestVerbs::VERB_NONE:
+		{
+			return TEXT("NONE");
+		}
+	}
+#endif
+	return TEXT("UNKNOWN");
+}
+
+
+UExternalRpcRegistry* UExternalRpcRegistry::GetInstance()
+{
+#if USE_RPC_REGISTRY
+	if (ObjectInstance == nullptr)
+	{		
+		ObjectInstance = NewObject<UExternalRpcRegistry>();
+		FParse::Value(FCommandLine::Get(), TEXT("rpcport="), ObjectInstance->PortToUse);
+		
+		TWeakObjectPtr<ThisClass> WeakThis(ObjectInstance);
+		// We always want the ListRegisteredRpcs route bound, no matter what.
+		UExternalRpcRegistry::GetInstance()->RegisterNewRoute(TEXT("ListRegisteredRpcs"), FHttpPath("/listrpcs"), EHttpServerRequestVerbs::VERB_GET,
+			[WeakThis](const FHttpServerRequest& Request, const FHttpResultCallback& OnComplete)
+		{
+			if (!WeakThis.IsValid()) { return false; }
+			return WeakThis->HttpListOpenRoutes(Request, OnComplete);
+		}, true);
+
+		ObjectInstance->AddToRoot();
+	}
+#endif
+	return ObjectInstance;
+}
+
+
+bool UExternalRpcRegistry::GetRegisteredRoute(FName RouteName, FExternalRouteInfo& OutRouteInfo)
+{
+	if (RegisteredRoutes.Find(RouteName))
+	{
+		OutRouteInfo.RouteName = RouteName;
+		OutRouteInfo.RoutePath = RegisteredRoutes[RouteName]->Path;
+		OutRouteInfo.RequestVerbs = RegisteredRoutes[RouteName]->Verbs;
+		return true;
+	}
+	return false;
+}
+
+void UExternalRpcRegistry::RegisterNewRoute(FName RouteName, const FHttpPath& HttpPath, const EHttpServerRequestVerbs& RequestVerbs, const FHttpRequestHandler& Handler, bool bOverrideIfBound /* = false */)
+{
+#if USE_RPC_REGISTRY
+	FExternalRouteInfo InRouteInfo;
+	InRouteInfo.RouteName = RouteName;
+	InRouteInfo.RoutePath = HttpPath;
+	InRouteInfo.RequestVerbs = RequestVerbs;
+	RegisterNewRoute(InRouteInfo, Handler, bOverrideIfBound);
+#endif
+}
+void UExternalRpcRegistry::RegisterNewRoute(FExternalRouteInfo InRouteInfo, const FHttpRequestHandler& Handler, bool bOverrideIfBound /* = false */)
+{
+#if USE_RPC_REGISTRY
+	TSharedPtr<IHttpRouter> HttpRouter = FHttpServerModule::Get().GetHttpRouter(PortToUse);
+
+	if (RegisteredRoutes.Find(InRouteInfo.RouteName))
+	{
+		if (!bOverrideIfBound)
+		{
+			UE_LOG(LogExternalRpcRegistry, Error, TEXT("Failed to bind route with friendly key %s - a route at location %s already exists."), *InRouteInfo.RouteName.ToString(), *InRouteInfo.RoutePath.GetPath());
+			return;
+		}
+		UE_LOG(LogExternalRpcRegistry, Log, TEXT("Overwriting route at friendly key %s - from %s to %s "), *InRouteInfo.RouteName.ToString(), *RegisteredRoutes[InRouteInfo.RouteName]->Path, *InRouteInfo.RoutePath.GetPath());
+		HttpRouter->UnbindRoute(RegisteredRoutes[InRouteInfo.RouteName]);
+	}
+	RegisteredRoutes.Add(InRouteInfo.RouteName, HttpRouter->BindRoute(InRouteInfo.RoutePath, InRouteInfo.RequestVerbs, Handler));
+#endif
+}
+
+
+void UExternalRpcRegistry::CleanUpRoute(FName RouteName, bool bFailIfUnbound /* = false */)
+{
+#if USE_RPC_REGISTRY
+	if (RegisteredRoutes.Find(RouteName))
+	{
+		TSharedPtr<IHttpRouter> HttpRouter = FHttpServerModule::Get().GetHttpRouter(PortToUse);
+		HttpRouter->UnbindRoute(RegisteredRoutes[RouteName]);
+	}
+	else
+	{
+		UE_LOG(LogExternalRpcRegistry, Warning, TEXT("Route name %s does not exist, could not unbind."), *RouteName.ToString());
+		check(!bFailIfUnbound);
+	}
+#endif
+}
+
+bool UExternalRpcRegistry::HttpListOpenRoutes(const FHttpServerRequest& Request, const FHttpResultCallback& OnComplete)
+{
+#if USE_RPC_REGISTRY
+	FString ResponseStr;
+	TArray<FName> OutRouteKeys;
+	RegisteredRoutes.GetKeys(OutRouteKeys);
+	TSharedRef<TJsonWriter<>> JsonWriter = TJsonWriterFactory<>::Create(&ResponseStr);
+	JsonWriter->WriteArrayStart();
+	for (const FName RouteKey : OutRouteKeys)
+	{
+		JsonWriter->WriteObjectStart();
+		JsonWriter->WriteValue(TEXT("name"), RouteKey.ToString());
+		JsonWriter->WriteValue(TEXT("route"), RegisteredRoutes[RouteKey]->Path);
+		JsonWriter->WriteValue(TEXT("verb"), GetHttpRouteVerbString(RegisteredRoutes[RouteKey]->Verbs));
+		JsonWriter->WriteObjectEnd();
+	}
+	JsonWriter->WriteArrayEnd();
+	JsonWriter->Close();
+	auto Response = FHttpServerResponse::Create(ResponseStr, TEXT("application/json"));
+	OnComplete(MoveTemp(Response));
+#endif
+	return true;
+}
+IMPLEMENT_MODULE(FDefaultModuleImpl, ExternalRpcRegistry);
\ No newline at end of file
diff --git a/Engine/Source/Runtime/ExternalRPCRegistry/Public/ExternalRpcRegistry.h b/Engine/Source/Runtime/ExternalRPCRegistry/Public/ExternalRpcRegistry.h
new file mode 100644
--- /dev/null
+++ b/Engine/Source/Runtime/ExternalRPCRegistry/Public/ExternalRpcRegistry.h
@@ -0,0 +1,77 @@
+// Copyright 1998-2019 Epic Games, Inc. All Rights Reserved.
+#pragma once
+
+#include "HttpResultCallback.h"
+#include "HttpRouteHandle.h"
+#include "HttpServerRequest.h"
+#include "IHttpRouter.h"
+#include "ExternalRpcRegistry.generated.h"
+
+DECLARE_LOG_CATEGORY_EXTERN(LogExternalRpcRegistry, Log, All);
+
+USTRUCT()
+struct FExternalRouteInfo
+{
+	GENERATED_BODY()
+public:
+	FName RouteName;
+	FHttpPath RoutePath;
+	EHttpServerRequestVerbs RequestVerbs;
+	FExternalRouteInfo()
+	{
+		RouteName = FName(TEXT(""));
+		RoutePath = FHttpPath();
+		RequestVerbs = EHttpServerRequestVerbs::VERB_NONE;
+	}
+	FExternalRouteInfo(FName InRouteName, FHttpPath InRoutePath, EHttpServerRequestVerbs InRequestVerbs)
+	{
+		RouteName = InRouteName;
+		RoutePath = InRoutePath;
+		RequestVerbs = InRequestVerbs;
+	}
+};
+
+
+/**
+ * This class is designed to be a singleton that handles registry, maintenance, and cleanup of any REST endpoints exposed on the process 
+ * for use in communicating with the process externally. 
+ */
+UCLASS()
+class EXTERNALRPCREGISTRY_API UExternalRpcRegistry : public UObject
+{
+	GENERATED_BODY()
+protected:
+	static UExternalRpcRegistry * ObjectInstance;
+	TMap<FName, FHttpRouteHandle> RegisteredRoutes;
+
+public:
+	static UExternalRpcRegistry * GetInstance();
+
+	int PortToUse = 11223;
+
+	/**
+	 * Try to get a route registered under given friendly name. Returns false if could not be found.
+	 */
+	bool GetRegisteredRoute(FName RouteName, FExternalRouteInfo& OutRouteInfo);
+
+	void RegisterNewRoute(FExternalRouteInfo InRouteInfo, const FHttpRequestHandler& Handler, bool bOverrideIfBound = false);
+
+	/**
+	 * Register a new route.
+	 * Will override existing routes if option is set, otherwise will error and fail to bind.
+	 */
+	void RegisterNewRoute(FName RouteName, const FHttpPath& HttpPath, const EHttpServerRequestVerbs& RequestVerbs, const FHttpRequestHandler& Handler, bool bOverrideIfBound = false);
+
+	/**
+	 * Clean up a route.
+	 * Can be set to fail if trying to unbind an unbound route.
+	 */
+	void CleanUpRoute(FName RouteName, bool bFailIfUnbound = false);
+
+	/**
+	 * Default Route Listing http call. Spits out all registered routes and describes them via a REST API call.
+	 * Always registered at /listrpcs GET by default
+	 */
+	bool HttpListOpenRoutes(const FHttpServerRequest& Request, const FHttpResultCallback& OnComplete);
+
+};
\ No newline at end of file
