Update usage #def name to be more in line with traditional usage ones.

diff --git a/Engine/Source/Runtime/ExternalRPCRegistry/Private/ExternalRpcRegistry.cpp b/Engine/Source/Runtime/ExternalRPCRegistry/Private/ExternalRpcRegistry.cpp
--- a/Engine/Source/Runtime/ExternalRPCRegistry/Private/ExternalRpcRegistry.cpp
+++ b/Engine/Source/Runtime/ExternalRPCRegistry/Private/ExternalRpcRegistry.cpp
@@ -9,7 +9,10 @@
 #include "Logging/LogMacros.h"
 
 #define USE_RPC_REGISTRY_IN_SHIPPING 0
-#define USE_RPC_REGISTRY (USE_RPC_REGISTRY_IN_SHIPPING || !UE_BUILD_SHIPPING )
+
+#ifndef WITH_RPC_REGISTRY
+#define WITH_RPC_REGISTRY (USE_RPC_REGISTRY_IN_SHIPPING || !UE_BUILD_SHIPPING )
+#endif
 
 DEFINE_LOG_CATEGORY(LogExternalRpcRegistry)
 
@@ -18,7 +21,7 @@ UExternalRpcRegistry* UExternalRpcRegistry::ObjectInstance = nullptr;
 
 FString GetHttpRouteVerbString(EHttpServerRequestVerbs InVerbs)
 {
-#if USE_RPC_REGISTRY
+#if WITH_RPC_REGISTRY
 	switch (InVerbs)
 	{
 		case EHttpServerRequestVerbs::VERB_POST:
@@ -53,7 +56,7 @@ FString GetHttpRouteVerbString(EHttpServerRequestVerbs InVerbs)
 
 UExternalRpcRegistry* UExternalRpcRegistry::GetInstance()
 {
-#if USE_RPC_REGISTRY
+#if WITH_RPC_REGISTRY
 	if (ObjectInstance == nullptr)
 	{		
 		ObjectInstance = NewObject<UExternalRpcRegistry>();
@@ -89,7 +92,7 @@ bool UExternalRpcRegistry::GetRegisteredRoute(FName RouteName, FExternalRouteInf
 
 void UExternalRpcRegistry::RegisterNewRoute(FName RouteName, const FHttpPath& HttpPath, const EHttpServerRequestVerbs& RequestVerbs, const FHttpRequestHandler& Handler, bool bOverrideIfBound /* = false */)
 {
-#if USE_RPC_REGISTRY
+#if WITH_RPC_REGISTRY
 	FExternalRouteInfo InRouteInfo;
 	InRouteInfo.RouteName = RouteName;
 	InRouteInfo.RoutePath = HttpPath;
@@ -99,7 +102,7 @@ void UExternalRpcRegistry::RegisterNewRoute(FName RouteName, const FHttpPath& Ht
 }
 void UExternalRpcRegistry::RegisterNewRoute(FExternalRouteInfo InRouteInfo, const FHttpRequestHandler& Handler, bool bOverrideIfBound /* = false */)
 {
-#if USE_RPC_REGISTRY
+#if WITH_RPC_REGISTRY
 	TSharedPtr<IHttpRouter> HttpRouter = FHttpServerModule::Get().GetHttpRouter(PortToUse);
 
 	if (RegisteredRoutes.Find(InRouteInfo.RouteName))
@@ -119,7 +122,7 @@ void UExternalRpcRegistry::RegisterNewRoute(FExternalRouteInfo InRouteInfo, cons
 
 void UExternalRpcRegistry::CleanUpRoute(FName RouteName, bool bFailIfUnbound /* = false */)
 {
-#if USE_RPC_REGISTRY
+#if WITH_RPC_REGISTRY
 	if (RegisteredRoutes.Find(RouteName))
 	{
 		TSharedPtr<IHttpRouter> HttpRouter = FHttpServerModule::Get().GetHttpRouter(PortToUse);
@@ -135,7 +138,7 @@ void UExternalRpcRegistry::CleanUpRoute(FName RouteName, bool bFailIfUnbound /*
 
 bool UExternalRpcRegistry::HttpListOpenRoutes(const FHttpServerRequest& Request, const FHttpResultCallback& OnComplete)
 {
-#if USE_RPC_REGISTRY
+#if WITH_RPC_REGISTRY
 	FString ResponseStr;
 	TArray<FName> OutRouteKeys;
 	RegisteredRoutes.GetKeys(OutRouteKeys);
