Engine functionality of RPC Categories.

diff --git a/Engine/Source/Runtime/ExternalRPCRegistry/Private/ExternalRpcRegistry.cpp b/Engine/Source/Runtime/ExternalRPCRegistry/Private/ExternalRpcRegistry.cpp
--- a/Engine/Source/Runtime/ExternalRPCRegistry/Private/ExternalRpcRegistry.cpp
+++ b/Engine/Source/Runtime/ExternalRPCRegistry/Private/ExternalRpcRegistry.cpp
@@ -53,6 +53,16 @@ FString GetHttpRouteVerbString(EHttpServerRequestVerbs InVerbs)
 	return TEXT("UNKNOWN");
 }
 
+bool UExternalRpcRegistry::IsActiveRpcCategory(FString InCategory)
+{
+#if WITH_RPC_REGISTRY
+	if (!ActiveRpcCategories.Num() || ActiveRpcCategories.Contains(InCategory))
+	{
+		return true;
+	}
+#endif
+	return false;
+}
 
 UExternalRpcRegistry* UExternalRpcRegistry::GetInstance()
 {
@@ -60,6 +70,19 @@ UExternalRpcRegistry* UExternalRpcRegistry::GetInstance()
 	if (ObjectInstance == nullptr)
 	{		
 		ObjectInstance = NewObject<UExternalRpcRegistry>();
+		FString InCommandLineValue;
+		if (FParse::Value(FCommandLine::Get(), TEXT("enabledrpccategories="), InCommandLineValue))
+		{
+			TArray<FString> Substrings;
+			InCommandLineValue.ParseIntoArray(Substrings, TEXT(","));
+			for (int i = 0; i < Substrings.Num(); i++)
+			{
+				if (!ObjectInstance->ActiveRpcCategories.Contains(Substrings[i]))
+				{
+					ObjectInstance->ActiveRpcCategories.Add(Substrings[i]);
+				}
+			}
+		}
 		FParse::Value(FCommandLine::Get(), TEXT("rpcport="), ObjectInstance->PortToUse);
 		
 		TWeakObjectPtr<ThisClass> WeakThis(ObjectInstance);
@@ -69,7 +92,8 @@ UExternalRpcRegistry* UExternalRpcRegistry::GetInstance()
 		{
 			if (!WeakThis.IsValid()) { return false; }
 			return WeakThis->HttpListOpenRoutes(Request, OnComplete);
-		}, true);
+			}, true,
+			true);
 
 		ObjectInstance->AddToRoot();
 	}
@@ -77,9 +101,9 @@ UExternalRpcRegistry* UExternalRpcRegistry::GetInstance()
 	return ObjectInstance;
 }
 
-
 bool UExternalRpcRegistry::GetRegisteredRoute(FName RouteName, FExternalRouteInfo& OutRouteInfo)
 {
+#if WITH_RPC_REGISTRY
 	if (RegisteredRoutes.Find(RouteName))
 	{
 		OutRouteInfo.RouteName = RouteName;
@@ -89,10 +113,11 @@ bool UExternalRpcRegistry::GetRegisteredRoute(FName RouteName, FExternalRouteInf
 		OutRouteInfo.InputExpectedFormat = RegisteredRoutes[RouteName].InputExpectedFormat;
 		return true;
 	}
+#endif
 	return false;
 }
 
-void UExternalRpcRegistry::RegisterNewRoute(FName RouteName, const FHttpPath& HttpPath, const EHttpServerRequestVerbs& RequestVerbs, const FHttpRequestHandler& Handler, bool bOverrideIfBound /* = false */, FString OptionalContentType /* = TEXT("")*/, FString OptionalExpectedFormat /*= TEXT("")*/)
+void UExternalRpcRegistry::RegisterNewRoute(FName RouteName, const FHttpPath& HttpPath, const EHttpServerRequestVerbs& RequestVerbs, const FHttpRequestHandler& Handler, bool bOverrideIfBound /* = false */, bool bIsAlwaysOn /* = false */, FString OptionalCategory /* = FString("Unknown") */, FString OptionalContentType /* = TEXT("")*/, FString OptionalExpectedFormat /*= TEXT("")*/)
 {
 #if WITH_RPC_REGISTRY
 	FExternalRouteInfo InRouteInfo;
@@ -101,12 +126,20 @@ void UExternalRpcRegistry::RegisterNewRoute(FName RouteName, const FHttpPath& Ht
 	InRouteInfo.RequestVerbs = RequestVerbs;
 	InRouteInfo.InputContentType = OptionalContentType;
 	InRouteInfo.InputExpectedFormat = OptionalExpectedFormat;
+	InRouteInfo.RpcCategory = OptionalCategory;
+	InRouteInfo.bAlwaysOn = bIsAlwaysOn;
 	RegisterNewRoute(InRouteInfo, Handler, bOverrideIfBound);
 #endif
 }
+
 void UExternalRpcRegistry::RegisterNewRoute(FExternalRouteInfo InRouteInfo, const FHttpRequestHandler& Handler, bool bOverrideIfBound /* = false */)
 {
 #if WITH_RPC_REGISTRY
+	if (!InRouteInfo.bAlwaysOn && !IsActiveRpcCategory(InRouteInfo.RpcCategory))
+	{
+		return;
+	}
+
 	TSharedPtr<IHttpRouter> HttpRouter = FHttpServerModule::Get().GetHttpRouter(PortToUse);
 
 	if (RegisteredRoutes.Find(InRouteInfo.RouteName))
diff --git a/Engine/Source/Runtime/ExternalRPCRegistry/Public/ExternalRpcRegistry.h b/Engine/Source/Runtime/ExternalRPCRegistry/Public/ExternalRpcRegistry.h
--- a/Engine/Source/Runtime/ExternalRPCRegistry/Public/ExternalRpcRegistry.h
+++ b/Engine/Source/Runtime/ExternalRPCRegistry/Public/ExternalRpcRegistry.h
@@ -18,7 +18,10 @@ public:
 	FHttpPath RoutePath;
 	EHttpServerRequestVerbs RequestVerbs;
 	FString InputContentType;
-	FString InputExpectedFormat;
+	FString InputExpectedFormat;	
+	FString RpcCategory;
+	bool bAlwaysOn;
+
 	FExternalRouteInfo()
 	{
 		RouteName = FName(TEXT(""));
@@ -26,14 +29,19 @@ public:
 		RequestVerbs = EHttpServerRequestVerbs::VERB_NONE;
 		InputContentType = TEXT("");
 		InputExpectedFormat = TEXT("");
+		RpcCategory = TEXT("Unknown");
+		bAlwaysOn = false;
 	}
-	FExternalRouteInfo(FName InRouteName, FHttpPath InRoutePath, EHttpServerRequestVerbs InRequestVerbs, FString InContentType = TEXT(""), FString InExpectedFormat = TEXT(""))
+
+	FExternalRouteInfo(FName InRouteName, FHttpPath InRoutePath, EHttpServerRequestVerbs InRequestVerbs, FString InCategory = TEXT("Unknown"), bool bInAlwaysOn = false, FString InContentType = TEXT(""), FString InExpectedFormat = TEXT(""))
 	{
 		RouteName = InRouteName;
 		RoutePath = InRoutePath;
 		RequestVerbs = InRequestVerbs;
 		InputContentType = InContentType;
 		InputExpectedFormat = InExpectedFormat;
+		RpcCategory = InCategory;
+		bAlwaysOn = false;
 	}
 };
 
@@ -68,12 +76,17 @@ class EXTERNALRPCREGISTRY_API UExternalRpcRegistry : public UObject
 protected:
 	static UExternalRpcRegistry * ObjectInstance;
 	TMap<FName, FExternalRouteDesc> RegisteredRoutes;
-
+	TArray<FString> ActiveRpcCategories;
 public:
 	static UExternalRpcRegistry * GetInstance();
 
 	int PortToUse = 11223;
 
+	/**
+	* Check if this Rpc is from a category that is meant to be enabled.
+	*/
+	bool IsActiveRpcCategory(FString InCategory);
+
 	/**
 	 * Try to get a route registered under given friendly name. Returns false if could not be found.
 	 */
@@ -85,7 +98,7 @@ public:
 	 * Register a new route.
 	 * Will override existing routes if option is set, otherwise will error and fail to bind.
 	 */
-	void RegisterNewRoute(FName RouteName, const FHttpPath& HttpPath, const EHttpServerRequestVerbs& RequestVerbs, const FHttpRequestHandler& Handler, bool bOverrideIfBound = false, FString OptionalContentType = TEXT(""), FString OptionalExpectedFormat = TEXT(""));
+	void RegisterNewRoute(FName RouteName, const FHttpPath& HttpPath, const EHttpServerRequestVerbs& RequestVerbs, const FHttpRequestHandler& Handler, bool bOverrideIfBound = false, bool bIsAlwaysOn = false, FString OptionalCategory = TEXT("Unknown"), FString OptionalContentType = TEXT(""), FString OptionalExpectedFormat = TEXT(""));
 
 	/**
 	 * Clean up a route.
