Clean up all registered routes as part of the RpcRegistry destructor to make sure we don't have any weird invalid function pointers rattling about.

diff --git a/Engine/Source/Runtime/ExternalRPCRegistry/Private/ExternalRpcRegistry.cpp b/Engine/Source/Runtime/ExternalRPCRegistry/Private/ExternalRpcRegistry.cpp
--- a/Engine/Source/Runtime/ExternalRPCRegistry/Private/ExternalRpcRegistry.cpp
+++ b/Engine/Source/Runtime/ExternalRPCRegistry/Private/ExternalRpcRegistry.cpp
@@ -66,6 +66,11 @@ bool UExternalRpcRegistry::IsActiveRpcCategory(FString InCategory)
 	return false;
 }
 
+UExternalRpcRegistry::~UExternalRpcRegistry()
+{
+	CleanUpAllRoutes();
+}
+
 UExternalRpcRegistry* UExternalRpcRegistry::GetInstance()
 {
 #if WITH_RPC_REGISTRY
@@ -179,6 +184,18 @@ void UExternalRpcRegistry::RegisterNewRoute(FExternalRouteInfo InRouteInfo, cons
 #endif
 }
 
+void UExternalRpcRegistry::CleanUpAllRoutes()
+{
+#if WITH_RPC_REGISTRY
+	TArray<FName> OutRouteKeys;
+	RegisteredRoutes.GetKeys(OutRouteKeys);
+	for (const FName& RouteKey : OutRouteKeys)
+	{
+		CleanUpRoute(RouteKey);
+	}
+#endif
+}
+
 void UExternalRpcRegistry::CleanUpRoute(FName RouteName, bool bFailIfUnbound /* = false */)
 {
 #if WITH_RPC_REGISTRY
diff --git a/Engine/Source/Runtime/ExternalRPCRegistry/Public/ExternalRpcRegistry.h b/Engine/Source/Runtime/ExternalRPCRegistry/Public/ExternalRpcRegistry.h
--- a/Engine/Source/Runtime/ExternalRPCRegistry/Public/ExternalRpcRegistry.h
+++ b/Engine/Source/Runtime/ExternalRPCRegistry/Public/ExternalRpcRegistry.h
@@ -116,6 +116,8 @@ protected:
 public:
 	static UExternalRpcRegistry * GetInstance();
 
+	~UExternalRpcRegistry();
+
 	int PortToUse = 11223;
 
 	/**
@@ -143,6 +145,10 @@ public:
 	UE_DEPRECATED(5.0, "RegisterNewRoute is deprecated when needing to add arguments, please use RegisterNewRouteWithArguments instead.")
 	void RegisterNewRoute(FName RouteName, const FHttpPath& HttpPath, const EHttpServerRequestVerbs& RequestVerbs, const FHttpRequestHandler& Handler, bool bOverrideIfBound = false, bool bIsAlwaysOn = false, FString OptionalCategory = TEXT("Unknown"), FString OptionalContentType = TEXT(""), FString OptionalExpectedFormat = TEXT(""));
 
+	/**
+	 * Clean up all routes - generally called as part of the destructor to make sure we don't have any function pointers dangling around.
+	 */
+	void CleanUpAllRoutes();
 
 	/**
 	 * Clean up a route.
