Handling for null arguments that might somehow get passed into External RPC Registry.

diff --git a/Engine/Source/Runtime/ExternalRPCRegistry/Private/ExternalRpcRegistry.cpp b/Engine/Source/Runtime/ExternalRPCRegistry/Private/ExternalRpcRegistry.cpp
--- a/Engine/Source/Runtime/ExternalRPCRegistry/Private/ExternalRpcRegistry.cpp
+++ b/Engine/Source/Runtime/ExternalRPCRegistry/Private/ExternalRpcRegistry.cpp
@@ -239,6 +239,10 @@ bool UExternalRpcRegistry::HttpListOpenRoutes(const FHttpServerRequest& Request,
 			JsonWriter->WriteArrayStart(TEXT("args"));
 			for (const FExternalRpcArgumentDesc* ArgDesc : RegisteredRoutes[RouteKey].ExpectedArguments)
 			{
+				if (ArgDesc == nullptr)
+				{
+					continue;
+				}
 				JsonWriter->WriteObjectStart();
 				JsonWriter->WriteValue(TEXT("name"), ArgDesc->Name);
 				JsonWriter->WriteValue(TEXT("type"), ArgDesc->Type);
