Add Swagger.json endpoint to ExternalRpcRegistry

diff --git a/Engine/Source/Runtime/ExternalRPCRegistry/Private/ExternalRpcRegistry.cpp b/Engine/Source/Runtime/ExternalRPCRegistry/Private/ExternalRpcRegistry.cpp
--- a/Engine/Source/Runtime/ExternalRPCRegistry/Private/ExternalRpcRegistry.cpp
+++ b/Engine/Source/Runtime/ExternalRPCRegistry/Private/ExternalRpcRegistry.cpp
@@ -4,6 +4,7 @@
 #include "HttpServerModule.h"
 #include "Misc/CommandLine.h"
 #include "Misc/Parse.h"
+#include "Misc/App.h"
 #include "HttpServerResponse.h"
 #include "Serialization/JsonWriter.h"
 #include "Logging/LogMacros.h"
@@ -93,10 +94,14 @@ UExternalRpcRegistry* UExternalRpcRegistry::GetInstance()
 		FParse::Value(FCommandLine::Get(), TEXT("rpcport="), ObjectInstance->PortToUse);
 		
 		FHttpRequestHandler ListRoutesRequestHandler = FHttpRequestHandler::CreateUObject(ObjectInstance, &ThisClass::HttpListOpenRoutes);
+		FHttpRequestHandler ListOASv3RequestHandler = FHttpRequestHandler::CreateUObject(ObjectInstance, &ThisClass::HttpListOASv3JSONRoutes);
 		TArray<FExternalRpcArgumentDesc> ArgumentArray;
 		// We always want the ListRegisteredRpcs route bound, no matter what.
 		ObjectInstance->RegisterNewRouteWithArguments(TEXT("ListRegisteredRpcs"), FHttpPath("/listrpcs"), EHttpServerRequestVerbs::VERB_GET,
-			ListRoutesRequestHandler, ArgumentArray,  true, true);
+			ListRoutesRequestHandler, ArgumentArray, true, true);
+		// /swagger.json escaped as %2e
+		ObjectInstance->RegisterNewRouteWithArguments(TEXT("ListSwaggerJson"), FHttpPath("/swagger.json"), EHttpServerRequestVerbs::VERB_GET,
+			ListOASv3RequestHandler, ArgumentArray, true, true);
 
 		ObjectInstance->AddToRoot();
 	}
@@ -174,6 +179,7 @@ void UExternalRpcRegistry::RegisterNewRoute(FExternalRouteInfo InRouteInfo, cons
 	RouteDesc.Handle = HttpRouter->BindRoute(InRouteInfo.RoutePath, InRouteInfo.RequestVerbs, Handler);
 	RouteDesc.InputContentType = InRouteInfo.InputContentType;
 	RouteDesc.ExpectedArguments = InRouteInfo.ExpectedArguments;
+	RouteDesc.RpcCategory = InRouteInfo.RpcCategory;
 	RegisteredRoutes.Add(InRouteInfo.RouteName, RouteDesc);
 #endif
 }
@@ -251,5 +257,139 @@ bool UExternalRpcRegistry::HttpListOpenRoutes(const FHttpServerRequest& Request,
 #endif
 	return true;
 }
+
+bool UExternalRpcRegistry::HttpListOASv3JSONRoutes(const FHttpServerRequest& Request, const FHttpResultCallback& OnComplete)
+{
+#if WITH_RPC_REGISTRY
+	FString ResponseStr;
+	TArray<FName> OutRouteKeys;
+	RegisteredRoutes.GetKeys(OutRouteKeys);
+	TSharedRef<TJsonWriter<>> JsonWriter = TJsonWriterFactory<>::Create(&ResponseStr);
+
+
+	JsonWriter->WriteObjectStart();
+
+	// Based on OpenApi Spec v3.0.0, update this string
+	// as necessary.
+	JsonWriter->WriteValue(TEXT("openapi"), TEXT("3.0.0"));
+
+	JsonWriter->WriteObjectStart(TEXT("info"));
+	JsonWriter->WriteValue(TEXT("title"), FString::Printf(TEXT("UE-%s - RPC API"), FApp::GetProjectName()));
+	JsonWriter->WriteValue(TEXT("description"), TEXT("Auto-generated Swagger API"));
+	JsonWriter->WriteValue(TEXT("version"), FApp::GetBuildVersion());
+	JsonWriter->WriteObjectEnd();
+
+	JsonWriter->WriteArrayStart(TEXT("servers"));
+	JsonWriter->WriteObjectStart();
+	JsonWriter->WriteValue(TEXT("url"), FString::Printf(TEXT("http://%s:%d"), TEXT("127.0.0.1"), PortToUse));
+	JsonWriter->WriteValue(TEXT("description"), TEXT("Default server access via localhost"));
+	JsonWriter->WriteObjectEnd();
+
+	/* TODO: Add external IP support. This section untested for now.
+	ISocketSubsystem* SocketSubsystem = ISocketSubsystem::Get();
+	if (SocketSubsystem)
+	{
+		TArray<TSharedRef<FInternetAddr>> Addresses = SocketSubsystem->GetLocalBindAddresses();
+		for (TSharedRef<FInternetAddr> InternetAddr : Addresses)
+		{
+			JsonWriter->WriteObjectStart();
+			JsonWriter->WriteValue(TEXT("url"), FString::Printf(TEXT("http://%s:%d"), InternetAddr->ToString(false), PortToUse));
+			JsonWriter->WriteValue(TEXT("description"), TEXT("Default server access via external ip"));
+			JsonWriter->WriteObjectEnd();
+		}
+	}
+	*/
+
+	JsonWriter->WriteArrayEnd();
+
+	JsonWriter->WriteObjectStart(TEXT("paths"));
+	for (const FName& RouteKey : OutRouteKeys)
+	{
+		JsonWriter->WriteObjectStart(RegisteredRoutes[RouteKey].Handle->Path);
+		JsonWriter->WriteObjectStart(GetHttpRouteVerbString(RegisteredRoutes[RouteKey].Handle->Verbs).ToLower());
+		JsonWriter->WriteValue(TEXT("summary"), RouteKey.ToString());
+		JsonWriter->WriteValue(TEXT("operationId"), RouteKey.ToString());
+
+		if (!RegisteredRoutes[RouteKey].RpcCategory.IsEmpty())
+		{
+			JsonWriter->WriteArrayStart(TEXT("tags"));
+			JsonWriter->WriteValue(RegisteredRoutes[RouteKey].RpcCategory);
+			JsonWriter->WriteArrayEnd();
+		}
+
+		// TODO: Ask C++ implementers to provide a description of their RPC call should do.
+		// We'll dump the InputContentType for now.
+		if (!RegisteredRoutes[RouteKey].InputContentType.IsEmpty())
+		{
+			JsonWriter->WriteValue(TEXT("description"), RegisteredRoutes[RouteKey].InputContentType);
+		}
+		else
+		{
+			JsonWriter->WriteValue(TEXT("description"), TEXT("No content type required to call this."));
+		}
+
+		if (!RegisteredRoutes[RouteKey].ExpectedArguments.IsEmpty())
+		{
+			JsonWriter->WriteObjectStart(TEXT("requestBody"));
+			JsonWriter->WriteObjectStart(TEXT("content"));
+			JsonWriter->WriteObjectStart(TEXT("application/json"));
+
+			JsonWriter->WriteObjectStart(TEXT("schema"));
+			JsonWriter->WriteValue(TEXT("type"), TEXT("object"));
+
+			JsonWriter->WriteObjectStart(TEXT("properties"));
+			TArray<FString> RequiredObjects;
+			for (const FExternalRpcArgumentDesc& ArgDesc : RegisteredRoutes[RouteKey].ExpectedArguments)
+			{
+				JsonWriter->WriteObjectStart(ArgDesc.Name);
+
+				//TODO: Provide better typing in RPC framework so we can auto-gen some values here.
+				JsonWriter->WriteValue(TEXT("description"), ArgDesc.Desc);
+				//JsonWriter->WriteValue(TEXT("type"), ArgDesc.Type);
+				JsonWriter->WriteValue(TEXT("type"), TEXT("string"));
+
+				if (!ArgDesc.bIsOptional)
+				{
+					RequiredObjects.Push(ArgDesc.Name);
+				}
+
+				JsonWriter->WriteObjectEnd();
+			}
+			JsonWriter->WriteObjectEnd();
+
+			if (RequiredObjects.Num() > 0)
+			{
+				JsonWriter->WriteArrayStart("required");
+				for (FString RequiredName : RequiredObjects)
+				{
+					JsonWriter->WriteValue(RequiredName);
+				}
+				JsonWriter->WriteArrayEnd();
+			}
+			JsonWriter->WriteObjectEnd();
+
+			JsonWriter->WriteObjectEnd();
+			JsonWriter->WriteObjectEnd();
+			JsonWriter->WriteObjectEnd();
+		}
+
+		JsonWriter->WriteObjectStart(TEXT("responses"));
+		JsonWriter->WriteObjectStart(TEXT("200"));
+		JsonWriter->WriteValue(TEXT("description"), TEXT("Successful return."));
+		JsonWriter->WriteObjectEnd();
+		JsonWriter->WriteObjectEnd();
+
+		JsonWriter->WriteObjectEnd();
+		JsonWriter->WriteObjectEnd();
+	}
+	JsonWriter->WriteObjectEnd();
+	JsonWriter->WriteObjectEnd();
+	JsonWriter->Close();
+	auto Response = FHttpServerResponse::Create(ResponseStr, TEXT("application/json"));
+	OnComplete(MoveTemp(Response));
+#endif
+	return true;
+}
+
 IMPLEMENT_MODULE(FDefaultModuleImpl, ExternalRpcRegistry);
 
diff --git a/Engine/Source/Runtime/ExternalRPCRegistry/Public/ExternalRpcRegistry.h b/Engine/Source/Runtime/ExternalRPCRegistry/Public/ExternalRpcRegistry.h
--- a/Engine/Source/Runtime/ExternalRPCRegistry/Public/ExternalRpcRegistry.h
+++ b/Engine/Source/Runtime/ExternalRPCRegistry/Public/ExternalRpcRegistry.h
@@ -106,13 +106,15 @@ struct FExternalRouteDesc
 public:
 	FHttpRouteHandle Handle;
 	FString InputContentType;
+	FString RpcCategory;
 	TArray<FExternalRpcArgumentDesc> ExpectedArguments;
 	FExternalRouteDesc() = default;
-	FExternalRouteDesc(FHttpRouteHandle InHandle, FString InContentType, TArray<FExternalRpcArgumentDesc> InArguments)
+	FExternalRouteDesc(FHttpRouteHandle InHandle, FString InContentType, TArray<FExternalRpcArgumentDesc> InArguments, FString InCategory = "")
 	{
 		Handle = InHandle;
 		InputContentType = InContentType;
 		ExpectedArguments = InArguments;
+		RpcCategory = InCategory;
 	}
 };
 
@@ -177,4 +179,10 @@ public:
 	 */
 	EXTERNALRPCREGISTRY_API bool HttpListOpenRoutes(const FHttpServerRequest& Request, const FHttpResultCallback& OnComplete);
 
+	/**
+	 * Route Listing http, formatted for Swagger OASv3. Spits out all registered routes and describes them via a REST API call.
+	 * Always registered at /swagger.json GET by default
+	 */
+	EXTERNALRPCREGISTRY_API bool HttpListOASv3JSONRoutes(const FHttpServerRequest& Request, const FHttpResultCallback& OnComplete);
+
 };
