Adding a request history to the RPC framework.

diff --git a/Engine/Source/Runtime/ExternalRPCRegistry/Private/ExternalRpcRegistry.cpp b/Engine/Source/Runtime/ExternalRPCRegistry/Private/ExternalRpcRegistry.cpp
--- a/Engine/Source/Runtime/ExternalRPCRegistry/Private/ExternalRpcRegistry.cpp
+++ b/Engine/Source/Runtime/ExternalRPCRegistry/Private/ExternalRpcRegistry.cpp
@@ -72,6 +72,19 @@ UExternalRpcRegistry::~UExternalRpcRegistry()
 	CleanUpAllRoutes();
 }
 
+bool UExternalRpcRegistry::IsEnabled()
+{
+#if WITH_RPC_REGISTRY
+	int32 RpcPort = 0;
+	// Not just returning this if because it'll cause static analysis issues for unreachable code in non-shipping
+	if (FParse::Value(FCommandLine::Get(), TEXT("rpcport="), RpcPort))
+	{
+		return true;
+	}
+#endif
+	return false;
+}
+
 UExternalRpcRegistry* UExternalRpcRegistry::GetInstance()
 {
 #if WITH_RPC_REGISTRY
@@ -92,16 +105,23 @@ UExternalRpcRegistry* UExternalRpcRegistry::GetInstance()
 			}
 		}
 		FParse::Value(FCommandLine::Get(), TEXT("rpcport="), ObjectInstance->PortToUse);
-		
+		FParse::Value(FCommandLine::Get(), TEXT("rpcledgersize="), ObjectInstance->RequestLedgerCapacity);
+
+		// We always want the ListRegisteredRpcs route bound, no matter what.
+
 		FHttpRequestHandler ListRoutesRequestHandler = FHttpRequestHandler::CreateUObject(ObjectInstance, &ThisClass::HttpListOpenRoutes);
-		FHttpRequestHandler ListOASv3RequestHandler = FHttpRequestHandler::CreateUObject(ObjectInstance, &ThisClass::HttpListOASv3JSONRoutes);
-		TArray<FExternalRpcArgumentDesc> ArgumentArray;
+		ObjectInstance->RegisterNewRoute(TEXT("ListRegisteredRpcs"), FHttpPath("/listrpcs"), EHttpServerRequestVerbs::VERB_GET,
+			ListRoutesRequestHandler, true, true);
+
+		FHttpRequestHandler PrintLedgerRequestHandler = FHttpRequestHandler::CreateUObject(ObjectInstance, &ThisClass::HttpPrintRequestLedger);
 		// We always want the ListRegisteredRpcs route bound, no matter what.
-		ObjectInstance->RegisterNewRouteWithArguments(TEXT("ListRegisteredRpcs"), FHttpPath("/listrpcs"), EHttpServerRequestVerbs::VERB_GET,
-			ListRoutesRequestHandler, ArgumentArray, true, true);
+		ObjectInstance->RegisterNewRoute(TEXT("GetRequestHistory"), FHttpPath("/requesthistory"), EHttpServerRequestVerbs::VERB_GET,
+			PrintLedgerRequestHandler, true, true);
+
+		FHttpRequestHandler ListOASv3RequestHandler = FHttpRequestHandler::CreateUObject(ObjectInstance, &ThisClass::HttpListOASv3JSONRoutes);
 		// /swagger.json escaped as %2e
-		ObjectInstance->RegisterNewRouteWithArguments(TEXT("ListSwaggerJson"), FHttpPath("/swagger.json"), EHttpServerRequestVerbs::VERB_GET,
-			ListOASv3RequestHandler, ArgumentArray, true, true);
+		ObjectInstance->RegisterNewRoute(TEXT("ListSwaggerJson"), FHttpPath("/swagger.json"), EHttpServerRequestVerbs::VERB_GET,
+			ListOASv3RequestHandler, true, true);
 
 		ObjectInstance->AddToRoot();
 	}
@@ -257,6 +277,48 @@ bool UExternalRpcRegistry::HttpListOpenRoutes(const FHttpServerRequest& Request,
 #endif
 	return true;
 }
+void UExternalRpcRegistry::AddRequestToLedger(const FHttpServerRequest& Request)
+{
+#if WITH_RPC_REGISTRY
+	if (Request.Headers.Find(TEXT("rpcname")))
+	{
+		FRpcLedgerEntry NewEntry;
+		NewEntry.RpcName = Request.Headers[TEXT("rpcname")][0];
+		FUTF8ToTCHAR WByteBuffer(reinterpret_cast<const ANSICHAR*>(Request.Body.GetData()), Request.Body.Num());
+		NewEntry.RequestBody = FString(WByteBuffer.Length(), WByteBuffer.Get());
+		NewEntry.RequestTime = FDateTime::UtcNow();
+		RequestLedger.Add(NewEntry);
+	}
+	// Reduce ledger to proper max size.
+	while (RequestLedger.Num() > RequestLedgerCapacity)
+	{
+		RequestLedger.RemoveAt(0);
+	}
+#endif
+}
+
+bool UExternalRpcRegistry::HttpPrintRequestLedger(const FHttpServerRequest& Request, const FHttpResultCallback& OnComplete)
+{
+#if WITH_RPC_REGISTRY
+	FString ResponseStr;
+	TSharedRef<TJsonWriter<>> JsonWriter = TJsonWriterFactory<>::Create(&ResponseStr);
+	JsonWriter->WriteArrayStart();
+	for (const FRpcLedgerEntry& LoggedRequest : RequestLedger)
+	{
+		JsonWriter->WriteObjectStart();
+		JsonWriter->WriteValue(TEXT("rpcname"), LoggedRequest.RpcName);
+		JsonWriter->WriteValue(TEXT("requesttimestamp"), LoggedRequest.RequestTime.ToString());
+		JsonWriter->WriteValue(TEXT("requestbody"), LoggedRequest.RequestBody);
+
+		JsonWriter->WriteObjectEnd();
+	}
+	JsonWriter->WriteArrayEnd();
+	JsonWriter->Close();
+	auto Response = FHttpServerResponse::Create(ResponseStr, TEXT("application/json"));
+	OnComplete(MoveTemp(Response));
+#endif
+	return true;
+}
 
 bool UExternalRpcRegistry::HttpListOASv3JSONRoutes(const FHttpServerRequest& Request, const FHttpResultCallback& OnComplete)
 {
diff --git a/Engine/Source/Runtime/ExternalRPCRegistry/Public/ExternalRpcRegistry.h b/Engine/Source/Runtime/ExternalRPCRegistry/Public/ExternalRpcRegistry.h
--- a/Engine/Source/Runtime/ExternalRPCRegistry/Public/ExternalRpcRegistry.h
+++ b/Engine/Source/Runtime/ExternalRPCRegistry/Public/ExternalRpcRegistry.h
@@ -118,6 +118,23 @@ public:
 	}
 };
 
+USTRUCT()
+struct FRpcLedgerEntry
+{
+	GENERATED_BODY()
+public:
+	FString RpcName;
+	FString RequestBody;
+	FDateTime RequestTime;
+	FRpcLedgerEntry() = default;
+	FRpcLedgerEntry(FString InName, FString InBody)
+	{
+		RpcName = InName;
+		RequestBody = InBody;
+		RequestTime = FDateTime::UtcNow();
+	}
+};
+
 /**
  * This class is designed to be a singleton that handles registry, maintenance, and cleanup of any REST endpoints exposed on the process 
  * for use in communicating with the process externally. 
@@ -130,12 +147,15 @@ protected:
 	static EXTERNALRPCREGISTRY_API UExternalRpcRegistry * ObjectInstance;
 	TMap<FName, FExternalRouteDesc> RegisteredRoutes;
 	TArray<FString> ActiveRpcCategories;
+	TArray<FRpcLedgerEntry> RequestLedger;
 public:
 	static EXTERNALRPCREGISTRY_API UExternalRpcRegistry * GetInstance();
+	static EXTERNALRPCREGISTRY_API bool IsEnabled();
 
 	EXTERNALRPCREGISTRY_API ~UExternalRpcRegistry();
 
 	int PortToUse = 11223;
+	int RequestLedgerCapacity = 10;
 
 	/**
 	* Check if this Rpc is from a category that is meant to be enabled.
@@ -159,7 +179,6 @@ public:
 	* Deprecated way to register a new route.
 	* Will override existing routes if option is set, otherwise will error and fail to bind.
 	 */
-	UE_DEPRECATED(5.0, "RegisterNewRoute is deprecated when needing to add arguments, please use RegisterNewRouteWithArguments instead.")
 	EXTERNALRPCREGISTRY_API void RegisterNewRoute(FName RouteName, const FHttpPath& HttpPath, const EHttpServerRequestVerbs& RequestVerbs, const FHttpRequestHandler& Handler, bool bOverrideIfBound = false, bool bIsAlwaysOn = false, FString OptionalCategory = TEXT("Unknown"), FString OptionalContentType = {}, FString OptionalExpectedFormat = {});
 
 	/**
@@ -178,6 +197,16 @@ public:
 	 * Always registered at /listrpcs GET by default
 	 */
 	EXTERNALRPCREGISTRY_API bool HttpListOpenRoutes(const FHttpServerRequest& Request, const FHttpResultCallback& OnComplete);
+	/**
+	* Records an RPC call into the ledger, plus trims the ledger if it is too large.
+	 */
+	EXTERNALRPCREGISTRY_API void AddRequestToLedger(const FHttpServerRequest& Request);
+
+	/**
+	 * Default request history call. Lists all calls recorded in the ledger, showing RPC name, request time, request body
+	 * Always registered at /requesthistory GET by default
+	 */
+	EXTERNALRPCREGISTRY_API bool HttpPrintRequestLedger(const FHttpServerRequest& Request, const FHttpResultCallback& OnComplete);
 
 	/**
 	 * Route Listing http, formatted for Swagger OASv3. Spits out all registered routes and describes them via a REST API call.
