Engine functionality for RPC Registration

diff --git a/Engine/Source/Runtime/ExternalRPCRegistry/ExternalRpcRegistry.Build.cs b/Engine/Source/Runtime/ExternalRPCRegistry/ExternalRpcRegistry.Build.cs
--- a/Engine/Source/Runtime/ExternalRPCRegistry/ExternalRpcRegistry.Build.cs
+++ b/Engine/Source/Runtime/ExternalRPCRegistry/ExternalRpcRegistry.Build.cs
@@ -11,8 +11,10 @@ public class ExternalRpcRegistry : ModuleRules
                 "Core",
                 "CoreUObject",
 				"Json",
-                "HTTPServer"
+				"HTTP",
+				"HTTPServer"
             }
         );
+		PublicDefinitions.Add("USE_RPC_REGISTRY_IN_SHIPPING=0");
     }
 }
diff --git a/Engine/Source/Runtime/ExternalRPCRegistry/Private/ExternalRpcRegistrationComponent.cpp b/Engine/Source/Runtime/ExternalRPCRegistry/Private/ExternalRpcRegistrationComponent.cpp
new file mode 100644
--- /dev/null
+++ b/Engine/Source/Runtime/ExternalRPCRegistry/Private/ExternalRpcRegistrationComponent.cpp
@@ -0,0 +1,195 @@
+// Copyright Epic Games, Inc. All Rights Reserved.
+
+#include "ExternalRpcRegistrationComponent.h"
+
+#include "HttpModule.h"
+#include "HttpServerResponse.h"
+#include "ExternalRpcRegistry.h"
+
+#include "Interfaces/IHttpRequest.h"
+
+
+#ifndef WITH_RPC_REGISTRY
+#define WITH_RPC_REGISTRY (USE_RPC_REGISTRY_IN_SHIPPING || !UE_BUILD_SHIPPING )
+#endif
+void UExternalRpcRegistrationComponent::DeregisterHttpCallbacks()
+{
+#if WITH_RPC_REGISTRY
+	for (const FName& RouteName : RegisteredRoutes)
+	{
+		UExternalRpcRegistry::GetInstance()->CleanUpRoute(RouteName);
+	}
+	RegisteredRoutes.Empty();
+	BroadcastRpcListChanged();
+#endif
+}
+
+void UExternalRpcRegistrationComponent::RegisterAlwaysOnHttpCallbacks()
+{
+#if WITH_RPC_REGISTRY
+	BroadcastRpcListChanged();
+#endif
+}
+
+bool UExternalRpcRegistrationComponent::HttpSendMessageToListener(FString MessageCategory, FString MessagePayload)
+{
+#if WITH_RPC_REGISTRY
+	if (ListenerAddress == TEXT("") || SenderID == TEXT(""))
+	{
+		// We don't have a listener.
+		return false;
+	}
+
+	FString RequestUri = FString::Printf(TEXT("http://%s/sendmessage"), *ListenerAddress);
+	FHttpModule& HttpModule = FHttpModule::Get();
+
+	FString RequestString;
+
+	TSharedRef<IHttpRequest, ESPMode::ThreadSafe> MessageRequest = HttpModule.CreateRequest();
+
+	TSharedRef<TJsonWriter<>> JsonWriter = TJsonWriterFactory<>::Create(&RequestString);
+	JsonWriter->WriteObjectStart();
+	JsonWriter->WriteValue(TEXT("senderid"), SenderID);
+	JsonWriter->WriteValue(TEXT("category"), MessageCategory);
+	JsonWriter->WriteValue(TEXT("payload"), MessagePayload);
+	JsonWriter->WriteObjectEnd();
+	JsonWriter->Close();
+
+	MessageRequest->SetVerb(TEXT("POST"));
+	MessageRequest->SetHeader(TEXT("Content-Type"), TEXT("application/json"));
+	MessageRequest->SetContentAsString(RequestString);
+	MessageRequest->SetURL(RequestUri);
+
+	bool bRequestProcessed = MessageRequest->ProcessRequest();
+
+	if (bRequestProcessed)
+	{
+		GLog->Log(TEXT("BotRPC"), ELogVerbosity::Log, FString::Printf(TEXT("Successfully sent message to %s from SenderID %s! body: %s"), *MessageRequest->GetURL(), *SenderID, *RequestString));
+	}
+	else
+	{
+		GLog->Log(TEXT("BotRPC"), ELogVerbosity::Error, FString::Printf(TEXT("Failed to send message to %s from SenderID %s!"), *MessageRequest->GetURL(), *SenderID));
+	}
+	return bRequestProcessed;
+#else
+	return false;
+#endif
+}
+void UExternalRpcRegistrationComponent::BroadcastRpcListChanged()
+{
+#if WITH_RPC_REGISTRY
+	HttpSendMessageToListener(TEXT("RpcRegistry"), TEXT("RpcListUpdated"));
+#endif
+}
+TUniquePtr<FHttpServerResponse> UExternalRpcRegistrationComponent::CreateSimpleResponse(bool bInWasSuccessful, FString InValue /*= ""*/)
+{
+	FString ResponseStr;
+	TSharedRef<TJsonWriter<>> JsonWriter = TJsonWriterFactory<>::Create(&ResponseStr);
+	JsonWriter->WriteObjectStart();
+	JsonWriter->WriteValue(TEXT("succeeded"), bInWasSuccessful);
+	JsonWriter->WriteValue(TEXT("value"), InValue);
+	JsonWriter->WriteObjectEnd();
+	JsonWriter->Close();
+	return FHttpServerResponse::Create(ResponseStr, TEXT("application/json"));
+}
+// Put any logic changes that should be applied to all of our API calls and no other API calls
+// inthis lambda
+FHttpRequestHandler UExternalRpcRegistrationComponent::CreateRouteHandle(TDelegate<bool(const FHttpServerRequest& Request, const FHttpResultCallback& OnComplete)> InFunc)
+{
+	FHttpRequestHandler OutHandler = nullptr;
+#if WITH_RPC_REGISTRY
+	OutHandler = FHttpRequestHandler::CreateLambda([this, InFunc](const FHttpServerRequest& Request, const FHttpResultCallback& OnComplete)
+		{
+			// If we have a security token enabled, make sure we're honoring it.
+			if (SecuritySecret != "")
+			{
+				if (!Request.Headers.Find(TEXT("authToken")))
+				{
+					TUniquePtr<FHttpServerResponse> Response = CreateSimpleResponse(false, TEXT("Missing authToken header"));
+					Response->Code = EHttpServerResponseCodes::Forbidden;
+					OnComplete(MoveTemp(Response));
+					return true;
+				}
+				TArray<FString> ExistingHeader = Request.Headers[TEXT("authToken")];
+				if (ExistingHeader[0] != SecuritySecret)
+				{
+					TUniquePtr<FHttpServerResponse> Response = CreateSimpleResponse(false, TEXT("Incorrect authToken header value"));
+					Response->Code = EHttpServerResponseCodes::Forbidden;
+					OnComplete(MoveTemp(Response));
+					return true;
+				}
+			}
+			UExternalRpcRegistry::GetInstance()->AddRequestToLedger(Request);
+			return InFunc.Execute(Request, OnComplete);
+
+		});
+#endif
+	return OutHandler;
+}
+
+
+void UExternalRpcRegistrationComponent::RegisterHttpCallback(FExternalRouteInfo InRouteInfo, const FHttpRequestHandler& Handler, bool bOverrideIfBound /* = false */)
+{
+#if WITH_RPC_REGISTRY
+	const FName RouteName = InRouteInfo.RouteName;
+	UExternalRpcRegistry::GetInstance()->RegisterNewRoute(MoveTemp(InRouteInfo), Handler, bOverrideIfBound);
+	RegisteredRoutes.Add(RouteName);
+#endif
+}
+
+void UExternalRpcRegistrationComponent::RegisterHttpCallback(FName RouteName, const FHttpPath& HttpPath, const EHttpServerRequestVerbs& RequestVerbs, const FHttpRequestHandler& Handler, bool bOverrideIfBound/* = false*/, FString OptionalCategory/* = TEXT("Unknown")*/, FString OptionalContentType/* = TEXT("")*/, TArray<FExternalRpcArgumentDesc> OptionalInArguments/* = TArray<FExternalRpcArgumentDesc>()*/)
+{
+#if WITH_RPC_REGISTRY
+	FExternalRouteInfo NewRouteInfo;
+	NewRouteInfo.RouteName = RouteName;
+	NewRouteInfo.RoutePath = HttpPath;
+	NewRouteInfo.RequestVerbs = RequestVerbs;
+	NewRouteInfo.InputContentType = MoveTemp(OptionalContentType);
+	NewRouteInfo.ExpectedArguments = MoveTemp(OptionalInArguments);
+	NewRouteInfo.RpcCategory = OptionalCategory;
+	RegisterHttpCallback(MoveTemp(NewRouteInfo), Handler, bOverrideIfBound);
+#endif
+}
+bool UExternalRpcRegistrationComponent::HttpUpdateIpOnListener(FString TargetName, FString NewIP)
+{
+#if WITH_RPC_REGISTRY
+	if (TargetName == TEXT("") || NewIP == TEXT(""))
+	{
+		// We don't have a listener.
+		GLog->Log(TEXT("BotRPC"), ELogVerbosity::Error, FString::Printf(TEXT("UpdateTargetIp Failed to send TargetName:  %s   NewIP: %s"), *TargetName, *NewIP));
+		return false;
+	}
+	FString RequestUri = FString::Printf(TEXT("http://%s/updateip"), *ListenerAddress);
+	FHttpModule& HttpModule = FHttpModule::Get();
+
+	FString RequestString;
+
+	TSharedRef<IHttpRequest, ESPMode::ThreadSafe> MessageRequest = HttpModule.CreateRequest();
+
+	TSharedRef<TJsonWriter<>> JsonWriter = TJsonWriterFactory<>::Create(&RequestString);
+	JsonWriter->WriteObjectStart();
+	JsonWriter->WriteValue(TEXT("target"), TargetName);
+	JsonWriter->WriteValue(TEXT("newip"), NewIP);
+	JsonWriter->WriteObjectEnd();
+	JsonWriter->Close();
+
+	MessageRequest->SetVerb(TEXT("POST"));
+	MessageRequest->SetHeader(TEXT("Content-Type"), TEXT("application/json"));
+	MessageRequest->SetContentAsString(RequestString);
+	MessageRequest->SetURL(RequestUri);
+
+	bool bRequestProcessed = MessageRequest->ProcessRequest();
+
+	if (bRequestProcessed)
+	{
+		GLog->Log(TEXT("BotRPC"), ELogVerbosity::Log, FString::Printf(TEXT("UpdateTargetIp: Successfully sent message to %s from SenderID %s! body: %s"), *MessageRequest->GetURL(), *SenderID, *RequestString));
+	}
+	else
+	{
+		GLog->Log(TEXT("BotRPC"), ELogVerbosity::Error, FString::Printf(TEXT("UpdateTargetIp Failed to send message to %s from SenderID %s!"), *MessageRequest->GetURL(), *SenderID));
+	}
+	return bRequestProcessed;
+#else
+return false;
+#endif
+}
\ No newline at end of file
diff --git a/Engine/Source/Runtime/ExternalRPCRegistry/Public/ExternalRpcRegistrationComponent.h b/Engine/Source/Runtime/ExternalRPCRegistry/Public/ExternalRpcRegistrationComponent.h
new file mode 100644
--- /dev/null
+++ b/Engine/Source/Runtime/ExternalRPCRegistry/Public/ExternalRpcRegistrationComponent.h
@@ -0,0 +1,48 @@
+// Copyright Epic Games, Inc. All Rights Reserved.
+
+#pragma once
+
+#include "HttpRequestHandler.h"
+#include "Serialization/JsonWriter.h"
+#include "ExternalRpcRegistry.h"
+
+
+#include "ExternalRpcRegistrationComponent.generated.h"
+class FJsonObject;
+enum class EHttpServerRequestVerbs : uint16;
+struct FExternalRouteInfo;
+struct FExternalRpcArgumentDesc;
+struct FHttpPath;
+struct FHttpServerRequest;
+struct FHttpServerResponse;
+struct FKey;
+UCLASS()
+class UExternalRpcRegistrationComponent : public UObject
+{
+
+	GENERATED_BODY()
+public:
+	TArray<FName> RegisteredRoutes;
+	FString SecuritySecret;
+
+
+	// Outgoing RPC info
+	FString SenderID;
+	FString ListenerAddress;
+
+	virtual void DeregisterHttpCallbacks();	
+	virtual void RegisterAlwaysOnHttpCallbacks();
+	void BroadcastRpcListChanged();
+	void RegisterHttpCallback(FExternalRouteInfo InRouteInfo, const FHttpRequestHandler& Handler, bool bOverrideIfBound = false);
+	void RegisterHttpCallback(FName RouteName, const FHttpPath& HttpPath, const EHttpServerRequestVerbs& RequestVerbs, const FHttpRequestHandler& Handler, bool bOverrideIfBound = false, FString OptionalCategory = TEXT(""), FString OptionalContentType = {}, TArray<FExternalRpcArgumentDesc> OptionalInArguments = TArray<FExternalRpcArgumentDesc>());
+
+
+	// Send Message to Listener RPC
+	bool HttpSendMessageToListener(FString MessageCategory, FString MessagePayload);
+	bool HttpUpdateIpOnListener(FString TargetName, FString MessagePayload);
+	FHttpRequestHandler CreateRouteHandle(TDelegate<bool(const FHttpServerRequest& Request, const FHttpResultCallback& OnComplete)>	 InFunc);
+
+	/** Creates a simple Http response with boolean and optional value */
+	TUniquePtr<FHttpServerResponse> CreateSimpleResponse(bool bInWasSuccessful, FString InValue = "");
+
+};
\ No newline at end of file
