Expose RPC Registratoin object as part of API

diff --git a/Engine/Source/Runtime/ExternalRPCRegistry/Public/ExternalRpcRegistrationComponent.h b/Engine/Source/Runtime/ExternalRPCRegistry/Public/ExternalRpcRegistrationComponent.h
--- a/Engine/Source/Runtime/ExternalRPCRegistry/Public/ExternalRpcRegistrationComponent.h
+++ b/Engine/Source/Runtime/ExternalRPCRegistry/Public/ExternalRpcRegistrationComponent.h
@@ -17,7 +17,7 @@ struct FHttpServerRequest;
 struct FHttpServerResponse;
 struct FKey;
 UCLASS()
-class UExternalRpcRegistrationComponent : public UObject
+class EXTERNALRPCREGISTRY_API UExternalRpcRegistrationComponent : public UObject
 {
 
 	GENERATED_BODY()
