Fixed polling for CheckFEngineLoopInitComplete
Converted MessagePayload: FEngineLoopInitComplete from old approach to the new

#rnx

diff --git a/Engine/Source/Runtime/ExternalRPCRegistry/Private/ExternalRpcRegistry.cpp b/Engine/Source/Runtime/ExternalRPCRegistry/Private/ExternalRpcRegistry.cpp
--- a/Engine/Source/Runtime/ExternalRPCRegistry/Private/ExternalRpcRegistry.cpp
+++ b/Engine/Source/Runtime/ExternalRPCRegistry/Private/ExternalRpcRegistry.cpp
@@ -86,24 +86,12 @@ bool UExternalRpcRegistry::IsEnabled()
 	return false;
 }
 
-void UExternalRpcRegistry::OnEngineInitComplete()
-{
-	bIsFEngineLoopInitComplete = true;
-}
-
-void UExternalRpcRegistry::BindEngineInitCompleteDelegate()
-{
-	FCoreDelegates::OnFEngineLoopInitComplete.AddUObject(this, &UExternalRpcRegistry::OnEngineInitComplete);
-}
-
-
 UExternalRpcRegistry* UExternalRpcRegistry::GetInstance()
 {
 #if WITH_RPC_REGISTRY
 	if (ObjectInstance == nullptr)
 	{		
 		ObjectInstance = NewObject<UExternalRpcRegistry>();
-		ObjectInstance->BindEngineInitCompleteDelegate();
 		FString InCommandLineValue;
 		if (FParse::Value(FCommandLine::Get(), TEXT("enabledrpccategories="), InCommandLineValue))
 		{
diff --git a/Engine/Source/Runtime/ExternalRPCRegistry/Public/ExternalRpcRegistry.h b/Engine/Source/Runtime/ExternalRPCRegistry/Public/ExternalRpcRegistry.h
--- a/Engine/Source/Runtime/ExternalRPCRegistry/Public/ExternalRpcRegistry.h
+++ b/Engine/Source/Runtime/ExternalRPCRegistry/Public/ExternalRpcRegistry.h
@@ -152,10 +152,6 @@ protected:
 public:
 	static EXTERNALRPCREGISTRY_API UExternalRpcRegistry * GetInstance();
 	static EXTERNALRPCREGISTRY_API bool IsEnabled();
-	EXTERNALRPCREGISTRY_API void OnEngineInitComplete();
-	EXTERNALRPCREGISTRY_API void BindEngineInitCompleteDelegate();
-	bool bIsFEngineLoopInitComplete;
-
 	EXTERNALRPCREGISTRY_API ~UExternalRpcRegistry();
 
 	int PortToUse = 11223;
