Add ability to go to /swagger/html.index to run relevant RPCs.
- Add /swagger/html.index endpoint that will serve as a interactive page.
- Update /swagger.json to return localhost or multihomeip if provided

#rnx

diff --git a/Engine/Source/Runtime/ExternalRPCRegistry/Private/ExternalRpcRegistry.cpp b/Engine/Source/Runtime/ExternalRPCRegistry/Private/ExternalRpcRegistry.cpp
--- a/Engine/Source/Runtime/ExternalRPCRegistry/Private/ExternalRpcRegistry.cpp
+++ b/Engine/Source/Runtime/ExternalRPCRegistry/Private/ExternalRpcRegistry.cpp
@@ -124,6 +124,10 @@ UExternalRpcRegistry* UExternalRpcRegistry::GetInstance()
 		ObjectInstance->RegisterNewRoute(TEXT("ListSwaggerJson"), FHttpPath("/swagger.json"), EHttpServerRequestVerbs::VERB_GET,
 			ListOASv3RequestHandler, true, true);
 
+		FHttpRequestHandler SwaggerUIHandler = FHttpRequestHandler::CreateUObject(ObjectInstance, &ThisClass::HttpSwaggerUI);
+		ObjectInstance->RegisterNewRoute(TEXT("SwaggerUIHTML"), FHttpPath("/swagger/index.html"), EHttpServerRequestVerbs::VERB_GET,
+			SwaggerUIHandler, true, true);
+
 		ObjectInstance->AddToRoot();
 	}
 #endif
@@ -321,6 +325,56 @@ bool UExternalRpcRegistry::HttpPrintRequestLedger(const FHttpServerRequest& Requ
 	return true;
 }
 
+bool UExternalRpcRegistry::HttpSwaggerUI(const FHttpServerRequest& Request, const FHttpResultCallback& OnComplete)
+{
+#if WITH_RPC_REGISTRY
+	/*TODO: Maybe read this from contents, but for now just bake it in code.*/
+	FString ResponseStr = R"(<!DOCTYPE html>
+<html lang="en">
+<head>
+  <meta charset="utf-8" />
+  <meta name="viewport" content="width=device-width, initial-scale=1" />
+  <meta
+    name="description"
+    content="SwaggerIU"
+  />
+  <title>SwaggerUI</title>
+  <link rel="stylesheet" href="https:[REDACTED_EMAIL]/swagger-ui.css" />
+</head>
+<body>
+  <div id="swagger-ui"></div>
+  <script src="https:[REDACTED_EMAIL]/swagger-ui-bundle.js" crossorigin></script>
+  <script src="https:[REDACTED_EMAIL]/swagger-ui-standalone-preset.js" crossorigin></script>
+  <script>
+    window.onload = () => {
+      window.ui = SwaggerUIBundle({
+        url: 'https://{0}:{1}/swagger.json',
+        dom_id: '#swagger-ui',
+      });
+    };
+  </script>
+</body>
+</html>)";
+
+	// Default to multihome address if provided, else use localhost.
+	FString Address;
+	if (FParse::Value(FCommandLine::Get(), TEXT("MULTIHOMEHTTP="), Address))
+	{
+	} else if (FParse::Value(FCommandLine::Get(), TEXT("MULTIHOME="), Address))
+	{
+	} else
+	{
+		Address = TEXT("127.0.0.1");
+	}
+
+	FStringFormatOrderedArguments Args{ Address,PortToUse };
+	FString ListFormat = FString::Format(*ResponseStr, Args);
+	auto Response = FHttpServerResponse::Create(ListFormat, TEXT("text/html"));
+	OnComplete(MoveTemp(Response));
+#endif
+	return true;
+}
+
 bool UExternalRpcRegistry::HttpListOASv3JSONRoutes(const FHttpServerRequest& Request, const FHttpResultCallback& OnComplete)
 {
 #if WITH_RPC_REGISTRY
@@ -343,25 +397,30 @@ bool UExternalRpcRegistry::HttpListOASv3JSONRoutes(const FHttpServerRequest& Req
 	JsonWriter->WriteObjectEnd();
 
 	JsonWriter->WriteArrayStart(TEXT("servers"));
-	JsonWriter->WriteObjectStart();
-	JsonWriter->WriteValue(TEXT("url"), FString::Printf(TEXT("http://%s:%d"), TEXT("127.0.0.1"), PortToUse));
-	JsonWriter->WriteValue(TEXT("description"), TEXT("Default server access via localhost"));
-	JsonWriter->WriteObjectEnd();
 
-	/* TODO: Add external IP support. This section untested for now.
-	ISocketSubsystem* SocketSubsystem = ISocketSubsystem::Get();
-	if (SocketSubsystem)
+	
+	TArray<FString> Addresses;
+
+	FString MultiHomeFromCommandLine;
+	if (FParse::Value(FCommandLine::Get(), TEXT("MULTIHOMEHTTP="), MultiHomeFromCommandLine))
 	{
-		TArray<TSharedRef<FInternetAddr>> Addresses = SocketSubsystem->GetLocalBindAddresses();
-		for (TSharedRef<FInternetAddr> InternetAddr : Addresses)
-		{
-			JsonWriter->WriteObjectStart();
-			JsonWriter->WriteValue(TEXT("url"), FString::Printf(TEXT("http://%s:%d"), InternetAddr->ToString(false), PortToUse));
-			JsonWriter->WriteValue(TEXT("description"), TEXT("Default server access via external ip"));
-			JsonWriter->WriteObjectEnd();
-		}
+		Addresses.Push(MultiHomeFromCommandLine);
+	} else if (FParse::Value(FCommandLine::Get(), TEXT("MULTIHOME="), MultiHomeFromCommandLine))
+	{
+		Addresses.Push(MultiHomeFromCommandLine);
+	}
+	
+	Addresses.Push(TEXT("127.0.0.1"));
+	Addresses.Push(TEXT("localhost"));
+
+	for (FString InternetAddr : Addresses)
+	{
+		JsonWriter->WriteObjectStart();
+		JsonWriter->WriteValue(TEXT("url"), FString::Printf(TEXT("http://%s:%d"), *InternetAddr, PortToUse));
+		JsonWriter->WriteValue(TEXT("description"), TEXT("Default server access ip"));
+		JsonWriter->WriteObjectEnd();
 	}
-	*/
+	
 
 	JsonWriter->WriteArrayEnd();
 
diff --git a/Engine/Source/Runtime/ExternalRPCRegistry/Public/ExternalRpcRegistry.h b/Engine/Source/Runtime/ExternalRPCRegistry/Public/ExternalRpcRegistry.h
--- a/Engine/Source/Runtime/ExternalRPCRegistry/Public/ExternalRpcRegistry.h
+++ b/Engine/Source/Runtime/ExternalRPCRegistry/Public/ExternalRpcRegistry.h
@@ -214,4 +214,11 @@ public:
 	 */
 	EXTERNALRPCREGISTRY_API bool HttpListOASv3JSONRoutes(const FHttpServerRequest& Request, const FHttpResultCallback& OnComplete);
 
+
+	/**
+	 * Route Showing SwaggerUI using /swagger.json to describe the current REST API.
+	 * Always registered at /swagger/html.index GET by default
+	 */
+	EXTERNALRPCREGISTRY_API bool HttpSwaggerUI(const FHttpServerRequest& Request, const FHttpResultCallback& OnComplete);
+
 };
