RPC: Fix bug where swagger.json was defaulting to http(s) when it was unavaible via WebUI

#tests Quickfix fom manual testing.

diff --git a/Engine/Source/Runtime/ExternalRPCRegistry/Private/ExternalRpcRegistry.cpp b/Engine/Source/Runtime/ExternalRPCRegistry/Private/ExternalRpcRegistry.cpp
--- a/Engine/Source/Runtime/ExternalRPCRegistry/Private/ExternalRpcRegistry.cpp
+++ b/Engine/Source/Runtime/ExternalRPCRegistry/Private/ExternalRpcRegistry.cpp
@@ -348,7 +348,7 @@ bool UExternalRpcRegistry::HttpSwaggerUI(const FHttpServerRequest& Request, cons
   <script>
     window.onload = () => {
       window.ui = SwaggerUIBundle({
-        url: 'https://{0}:{1}/swagger.json',
+        url: 'http://{0}:{1}/swagger.json',
         dom_id: '#swagger-ui',
       });
     };
