Add the ability for rpc error responses to report as fatal, if necessary.

diff --git a/Engine/Source/Programs/AutomationTool/Gauntlet/Framework/RpcFramework/Gauntlet.RpcLibrary.cs b/Engine/Source/Programs/AutomationTool/Gauntlet/Framework/RpcFramework/Gauntlet.RpcLibrary.cs
--- a/Engine/Source/Programs/AutomationTool/Gauntlet/Framework/RpcFramework/Gauntlet.RpcLibrary.cs
+++ b/Engine/Source/Programs/AutomationTool/Gauntlet/Framework/RpcFramework/Gauntlet.RpcLibrary.cs
@@ -22,6 +22,7 @@ namespace Gauntlet
 		{
 			public bool Succeeded;
 			public string Value;
+			public bool Fatal;
 		}
 
 		public static bool ValidateTarget(RpcTarget InTarget)
diff --git a/Engine/Source/Runtime/ExternalRPCRegistry/Private/ExternalRpcRegistrationComponent.cpp b/Engine/Source/Runtime/ExternalRPCRegistry/Private/ExternalRpcRegistrationComponent.cpp
--- a/Engine/Source/Runtime/ExternalRPCRegistry/Private/ExternalRpcRegistrationComponent.cpp
+++ b/Engine/Source/Runtime/ExternalRPCRegistry/Private/ExternalRpcRegistrationComponent.cpp
@@ -81,13 +81,14 @@ void UExternalRpcRegistrationComponent::BroadcastRpcListChanged()
 	HttpSendMessageToListener(TEXT("RpcRegistry"), TEXT("RpcListUpdated"));
 #endif
 }
-TUniquePtr<FHttpServerResponse> UExternalRpcRegistrationComponent::CreateSimpleResponse(bool bInWasSuccessful, FString InValue /*= ""*/)
+TUniquePtr<FHttpServerResponse> UExternalRpcRegistrationComponent::CreateSimpleResponse(bool bInWasSuccessful, FString InValue /*= ""*/, bool bInFatal /* = false*/)
 {
 	FString ResponseStr;
 	TSharedRef<TJsonWriter<>> JsonWriter = TJsonWriterFactory<>::Create(&ResponseStr);
 	JsonWriter->WriteObjectStart();
 	JsonWriter->WriteValue(TEXT("succeeded"), bInWasSuccessful);
 	JsonWriter->WriteValue(TEXT("value"), InValue);
+	JsonWriter->WriteValue(TEXT("fatal"), bInFatal);
 	JsonWriter->WriteObjectEnd();
 	JsonWriter->Close();
 	return FHttpServerResponse::Create(ResponseStr, TEXT("application/json"));
diff --git a/Engine/Source/Runtime/ExternalRPCRegistry/Public/ExternalRpcRegistrationComponent.h b/Engine/Source/Runtime/ExternalRPCRegistry/Public/ExternalRpcRegistrationComponent.h
--- a/Engine/Source/Runtime/ExternalRPCRegistry/Public/ExternalRpcRegistrationComponent.h
+++ b/Engine/Source/Runtime/ExternalRPCRegistry/Public/ExternalRpcRegistrationComponent.h
@@ -100,8 +100,9 @@ public:
 	 *
 	 * @param  bInWasSuccessful   Whether or not the Rpc completed successfully
 	 * @param  InValue  Optional result description string that will be returned with the success bool
+	 * @param  bInFatal  Whether or not the response should be considered fatal and not worth retrying.
 	 */
-	UE_API TUniquePtr<FHttpServerResponse> CreateSimpleResponse(bool bInWasSuccessful, FString InValue = "");
+	UE_API TUniquePtr<FHttpServerResponse> CreateSimpleResponse(bool bInWasSuccessful, FString InValue = "", bool bInFatal = false);
 
 
 };
