[ExternalRpcRegistry] Add some logging for when a route name is bound, to match the logs emitted for when a route name is unbound.

diff --git a/Engine/Source/Runtime/ExternalRPCRegistry/Private/ExternalRpcRegistry.cpp b/Engine/Source/Runtime/ExternalRPCRegistry/Private/ExternalRpcRegistry.cpp
--- a/Engine/Source/Runtime/ExternalRPCRegistry/Private/ExternalRpcRegistry.cpp
+++ b/Engine/Source/Runtime/ExternalRPCRegistry/Private/ExternalRpcRegistry.cpp
@@ -206,6 +206,7 @@ void UExternalRpcRegistry::RegisterNewRoute(FExternalRouteInfo InRouteInfo, cons
 	RouteDesc.ExpectedArguments = InRouteInfo.ExpectedArguments;
 	RouteDesc.RpcCategory = InRouteInfo.RpcCategory;
 	RegisteredRoutes.Add(InRouteInfo.RouteName, RouteDesc);
+	UE_LOG(LogExternalRpcRegistry, Log, TEXT("Route name %s was bound!"), *InRouteInfo.RouteName.ToString());
 #endif
 }
 
