- Made a FRpcConfig class that can load settings from Engine.ini, for multiple ports for communication between processes (so editor and client on the same machine can easily run with different ports). Will use -rpcport as usual for standard operation

diff --git a/Engine/Source/Runtime/ExternalRPCRegistry/Private/ExternalRpcRegistry.cpp b/Engine/Source/Runtime/ExternalRPCRegistry/Private/ExternalRpcRegistry.cpp
--- a/Engine/Source/Runtime/ExternalRPCRegistry/Private/ExternalRpcRegistry.cpp
+++ b/Engine/Source/Runtime/ExternalRPCRegistry/Private/ExternalRpcRegistry.cpp
@@ -3,6 +3,7 @@
 #include "ExternalRpcRegistry.h"
 #include "HttpServerModule.h"
 #include "Misc/CommandLine.h"
+#include "Misc/ConfigCacheIni.h"
 #include "Misc/Parse.h"
 #include "Misc/App.h"
 #include "Misc/CoreDelegates.h"
@@ -73,12 +74,78 @@ UExternalRpcRegistry::~UExternalRpcRegistry()
 	CleanUpAllRoutes();
 }
 
+
+const FRpcConfig& UExternalRpcRegistry::GetRpcConfig()
+{
+	static FRpcConfig Config;
+	static bool bHasInitialized = false;
+
+	if (bHasInitialized)
+	{
+		return Config;
+	}
+	bHasInitialized = true;
+
+	FString ConfigName;
+	FParse::Value(FCommandLine::Get(), TEXT("-rpcconfig="), ConfigName);
+
+	if (FParse::Value(FCommandLine::Get(), TEXT("-rpcport="), Config.ActiveHostPort))
+	{
+		checkf(ConfigName.Len() == 0, TEXT("It is not supported to run with -rpcport and -rpcconfig"));
+
+		// put it into Default, just in case
+		Config.PortMapping.Add(TEXT("Default"), Config.ActiveHostPort);
+	}
+	// if we don't have -rpcconfig, look for a -rpcconfig=FooRpc string 
+	else if (ConfigName.Len() > 0)
+	{
+		const FConfigSection* Section = GConfig->GetSection(*ConfigName, false, GEngineIni);
+		if (Section)
+		{
+			FString SelectorType;
+			// put all the entries into the mapping
+			for (const TPair<FName, FConfigValue>& Pair : *Section)
+			{
+				// look for a special selector key
+				if (Pair.Key == TEXT("__Selector"))
+				{
+					SelectorType = Pair.Value.GetValue();
+					continue;
+				}
+
+				int Port = FCString::Atoi(*Pair.Value.GetValue());
+				if (Port > 0)
+				{
+					Config.PortMapping.Add(Pair.Key.ToString(), Port);
+				}
+			}
+
+			// now look for the active port (or look for a default entry if none given)
+			FString ActiveName = "Default";
+
+			if (SelectorType == TEXT("RuntimeMode"))
+			{
+				ActiveName = GIsEditor ? TEXT("Editor") :
+					GIsServer ? TEXT("Server") : TEXT("Client");
+			}
+			else
+			{
+				// look for rpcname, or use Default if none
+				FParse::Value(FCommandLine::Get(), TEXT("-rpcname="), ActiveName);
+			}
+
+			// look up the active port, or use 0 for missing port (0 isn't a valid port)
+			Config.ActiveHostPort = Config.PortMapping.FindRef(ActiveName);
+		}
+	}
+
+	return Config;
+}
+
 bool UExternalRpcRegistry::IsEnabled()
 {
 #if WITH_RPC_REGISTRY
-	int32 RpcPort = 0;
-	// Not just returning this if because it'll cause static analysis issues for unreachable code in non-shipping
-	if (FParse::Value(FCommandLine::Get(), TEXT("rpcport="), RpcPort))
+	if (GetRpcConfig().ActiveHostPort > 0)
 	{
 		return true;
 	}
@@ -105,7 +172,14 @@ UExternalRpcRegistry* UExternalRpcRegistry::GetInstance()
 				}
 			}
 		}
-		FParse::Value(FCommandLine::Get(), TEXT("rpcport="), ObjectInstance->PortToUse);
+		
+		// look for override from commandline (we want a value here even if UExternalRpcRegistry::IsEnabled() returns false, so the 
+		// routes below, and in other code, don't trigger an assert on registration)
+		int CmdLinePort = GetRpcConfig().ActiveHostPort;
+		if (CmdLinePort > 0)
+		{
+			ObjectInstance->PortToUse = CmdLinePort;
+		}
 		FParse::Value(FCommandLine::Get(), TEXT("rpcledgersize="), ObjectInstance->RequestLedgerCapacity);
 
 		// We always want the ListRegisteredRpcs route bound, no matter what.
diff --git a/Engine/Source/Runtime/ExternalRPCRegistry/Public/ExternalRpcRegistry.h b/Engine/Source/Runtime/ExternalRPCRegistry/Public/ExternalRpcRegistry.h
--- a/Engine/Source/Runtime/ExternalRPCRegistry/Public/ExternalRpcRegistry.h
+++ b/Engine/Source/Runtime/ExternalRPCRegistry/Public/ExternalRpcRegistry.h
@@ -135,6 +135,20 @@ public:
 	}
 };
 
+
+USTRUCT()
+struct FRpcConfig
+{
+	GENERATED_BODY()
+	FRpcConfig() = default;
+
+	TMap<FString, int> PortMapping;
+
+	// it is valid to have mappings, but no active port, in the case of a client with multiple endpoint hosts
+	int ActiveHostPort = 0;
+};
+
+
 /**
  * This class is designed to be a singleton that handles registry, maintenance, and cleanup of any REST endpoints exposed on the process 
  * for use in communicating with the process externally. 
@@ -152,6 +166,7 @@ protected:
 public:
 	static EXTERNALRPCREGISTRY_API UExternalRpcRegistry * GetInstance();
 	static EXTERNALRPCREGISTRY_API bool IsEnabled();
+	static EXTERNALRPCREGISTRY_API const FRpcConfig& GetRpcConfig();
 	EXTERNALRPCREGISTRY_API ~UExternalRpcRegistry();
 
 	int PortToUse = 11223;
