- Fixed a comment that had incorrect cmdline option name

diff --git a/Engine/Source/Runtime/ExternalRPCRegistry/Private/ExternalRpcRegistry.cpp b/Engine/Source/Runtime/ExternalRPCRegistry/Private/ExternalRpcRegistry.cpp
--- a/Engine/Source/Runtime/ExternalRPCRegistry/Private/ExternalRpcRegistry.cpp
+++ b/Engine/Source/Runtime/ExternalRPCRegistry/Private/ExternalRpcRegistry.cpp
@@ -96,7 +96,7 @@ const FRpcConfig& UExternalRpcRegistry::GetRpcConfig()
 		// put it into Default, just in case
 		Config.PortMapping.Add(TEXT("Default"), Config.ActiveHostPort);
 	}
-	// if we don't have -rpcconfig, look for a -rpcconfig=FooRpc string 
+	// if we don't have -rpcport, look for a -rpcconfig=FooRpc string 
 	else if (ConfigName.Len() > 0)
 	{
 		const FConfigSection* Section = GConfig->GetSection(*ConfigName, false, GEngineIni);
