Make it so we can't start a new httplistener or register routes if rpcport is not passed in to guard against accidental activations.

diff --git a/Engine/Source/Runtime/ExternalRPCRegistry/Private/ExternalRpcRegistry.cpp b/Engine/Source/Runtime/ExternalRPCRegistry/Private/ExternalRpcRegistry.cpp
--- a/Engine/Source/Runtime/ExternalRPCRegistry/Private/ExternalRpcRegistry.cpp
+++ b/Engine/Source/Runtime/ExternalRPCRegistry/Private/ExternalRpcRegistry.cpp
@@ -227,6 +227,10 @@ bool UExternalRpcRegistry::GetRegisteredRoute(FName RouteName, FExternalRouteInf
 void UExternalRpcRegistry::RegisterNewRouteWithArguments(FName RouteName, const FHttpPath& HttpPath, const EHttpServerRequestVerbs& RequestVerbs, const FHttpRequestHandler& Handler, TArray<FExternalRpcArgumentDesc> InArguments, bool bOverrideIfBound /* = false */, bool bIsAlwaysOn /* = false */, FString OptionalCategory /* = FString("Unknown") */, FString OptionalContentType /* = TEXT("")*/)
 {
 #if WITH_RPC_REGISTRY
+	if (!IsEnabled())
+	{
+		return;
+	}
 	FExternalRouteInfo InRouteInfo;
 	InRouteInfo.RouteName = RouteName;
 	InRouteInfo.RoutePath = HttpPath;
@@ -243,6 +247,10 @@ void UExternalRpcRegistry::RegisterNewRouteWithArguments(FName RouteName, const
 void UExternalRpcRegistry::RegisterNewRoute(FName RouteName, const FHttpPath& HttpPath, const EHttpServerRequestVerbs& RequestVerbs, const FHttpRequestHandler& Handler, bool bOverrideIfBound /* = false */, bool bIsAlwaysOn /* = false */, FString OptionalCategory /* = FString("Unknown") */, FString OptionalContentType /* = TEXT("")*/, FString OptionalExpectedFormat /*= TEXT("")*/)
 {
 #if WITH_RPC_REGISTRY
+	if (!IsEnabled())
+	{
+		return;
+	}
 	FExternalRouteInfo InRouteInfo;
 	InRouteInfo.RouteName = RouteName;
 	InRouteInfo.RoutePath = HttpPath;
@@ -257,6 +265,11 @@ void UExternalRpcRegistry::RegisterNewRoute(FName RouteName, const FHttpPath& Ht
 void UExternalRpcRegistry::RegisterNewRoute(FExternalRouteInfo InRouteInfo, const FHttpRequestHandler& Handler, bool bOverrideIfBound /* = false */)
 {
 #if WITH_RPC_REGISTRY
+	if (!IsEnabled())
+	{
+		return;
+	}
+
 	if (!InRouteInfo.bAlwaysOn && !IsActiveRpcCategory(InRouteInfo.RpcCategory))
 	{
 		return;
diff --git a/Engine/Source/Runtime/ExternalRPCRegistry/Public/ExternalRpcRegistry.h b/Engine/Source/Runtime/ExternalRPCRegistry/Public/ExternalRpcRegistry.h
--- a/Engine/Source/Runtime/ExternalRPCRegistry/Public/ExternalRpcRegistry.h
+++ b/Engine/Source/Runtime/ExternalRPCRegistry/Public/ExternalRpcRegistry.h
@@ -169,7 +169,7 @@ public:
 	static EXTERNALRPCREGISTRY_API const FRpcConfig& GetRpcConfig();
 	EXTERNALRPCREGISTRY_API ~UExternalRpcRegistry();
 
-	int PortToUse = 11223;
+	int PortToUse = 0;
 	int RequestLedgerCapacity = 10;
 
 	/**
