- Added support for -RpcClientNum to handle multiple clients with unique ports, and they just increment after the "client" port in the .ini, when using -RpcConfig=foo

diff --git a/Engine/Source/Runtime/ExternalRPCRegistry/Private/ExternalRpcRegistry.cpp b/Engine/Source/Runtime/ExternalRPCRegistry/Private/ExternalRpcRegistry.cpp
--- a/Engine/Source/Runtime/ExternalRPCRegistry/Private/ExternalRpcRegistry.cpp
+++ b/Engine/Source/Runtime/ExternalRPCRegistry/Private/ExternalRpcRegistry.cpp
@@ -136,6 +136,16 @@ const FRpcConfig& UExternalRpcRegistry::GetRpcConfig()
 
 			// look up the active port, or use 0 for missing port (0 isn't a valid port)
 			Config.ActiveHostPort = Config.PortMapping.FindRef(ActiveName);
+
+			// allow for multiple clients, so client nums >= 1 we append the num. client 0 is just called Client
+			if (ActiveName == TEXT("Client"))
+			{
+				int ClientNum = 0;
+				if (FParse::Value(FCommandLine::Get(), TEXT("-RpcClientNum="), ClientNum) && ClientNum >= 0 && ClientNum < 1000)
+				{
+					Config.ActiveHostPort += ClientNum;
+				}
+			}
 		}
 	}
 
