- Added support for multiple listening ports for Rpcs

diff --git a/Engine/Source/Runtime/ExternalRPCRegistry/Private/ExternalRpcRegistry.cpp b/Engine/Source/Runtime/ExternalRPCRegistry/Private/ExternalRpcRegistry.cpp
--- a/Engine/Source/Runtime/ExternalRPCRegistry/Private/ExternalRpcRegistry.cpp
+++ b/Engine/Source/Runtime/ExternalRPCRegistry/Private/ExternalRpcRegistry.cpp
@@ -86,18 +86,22 @@ const FRpcConfig& UExternalRpcRegistry::GetRpcConfig()
 	}
 	bHasInitialized = true;
 
-	FString ConfigName;
-	FParse::Value(FCommandLine::Get(), TEXT("-rpcconfig="), ConfigName);
 
-	if (FParse::Value(FCommandLine::Get(), TEXT("-rpcport="), Config.ActiveHostPort))
+	// look for multiple -rpcport or -rpcconfig commandline options
+	int PortNum;
+	for (const TCHAR* Cursor = FCommandLine::Get(); (Cursor != nullptr) && FParse::Value(Cursor, TEXT("-rpcport="), PortNum, &Cursor); )
 	{
-		checkf(ConfigName.Len() == 0, TEXT("It is not supported to run with -rpcport and -rpcconfig"));
-
-		// put it into Default, just in case
-		Config.PortMapping.Add(TEXT("Default"), Config.ActiveHostPort);
+		Config.ActiveHostPorts.AddUnique(PortNum);
+		// set one as Default, just in case something asks
+		if (!Config.PortMapping.Contains(TEXT("Default")))
+		{
+			Config.PortMapping.Add(TEXT("Default"), PortNum);
+		}
 	}
-	// if we don't have -rpcport, look for a -rpcconfig=FooRpc string 
-	else if (ConfigName.Len() > 0)
+
+	// look for a -rpcconfig=FooRpc string 
+	FString ConfigName;
+	for (const TCHAR* Cursor = FCommandLine::Get(); (Cursor != nullptr) && FParse::Value(Cursor, TEXT("-rpcconfig="), ConfigName, false, &Cursor); )
 	{
 		const FConfigSection* Section = GConfig->GetSection(*ConfigName, false, GEngineIni);
 		if (Section)
@@ -135,16 +139,21 @@ const FRpcConfig& UExternalRpcRegistry::GetRpcConfig()
 			}
 
 			// look up the active port, or use 0 for missing port (0 isn't a valid port)
-			Config.ActiveHostPort = Config.PortMapping.FindRef(ActiveName);
+			PortNum = Config.PortMapping.FindRef(ActiveName);
 
-			// allow for multiple clients, so client nums >= 1 we append the num. client 0 is just called Client
-			if (ActiveName == TEXT("Client"))
+			if (PortNum > 0)
 			{
-				int ClientNum = 0;
-				if (FParse::Value(FCommandLine::Get(), TEXT("-RpcClientNum="), ClientNum) && ClientNum >= 0 && ClientNum < 1000)
+				// allow for multiple clients, so client nums >= 1 we append the num. client 0 is just called Client
+				if (ActiveName == TEXT("Client"))
 				{
-					Config.ActiveHostPort += ClientNum;
+					int ClientNum = 0;
+					if (FParse::Value(FCommandLine::Get(), TEXT("-RpcClientNum="), ClientNum) && ClientNum >= 0 && ClientNum < 1000)
+					{
+						PortNum += ClientNum;
+					}
 				}
+
+				Config.ActiveHostPorts.AddUnique(PortNum);
 			}
 		}
 	}
@@ -155,7 +164,7 @@ const FRpcConfig& UExternalRpcRegistry::GetRpcConfig()
 bool UExternalRpcRegistry::IsEnabled()
 {
 #if WITH_RPC_REGISTRY
-	if (GetRpcConfig().ActiveHostPort > 0)
+	if (GetRpcConfig().ActiveHostPorts.Num() > 0)
 	{
 		return true;
 	}
@@ -185,11 +194,12 @@ UExternalRpcRegistry* UExternalRpcRegistry::GetInstance()
 		
 		// look for override from commandline (we want a value here even if UExternalRpcRegistry::IsEnabled() returns false, so the 
 		// routes below, and in other code, don't trigger an assert on registration)
-		int CmdLinePort = GetRpcConfig().ActiveHostPort;
-		if (CmdLinePort > 0)
+		ObjectInstance->PortsToUse = GetRpcConfig().ActiveHostPorts;
+		if (ObjectInstance->PortsToUse.Num() == 0)
 		{
-			ObjectInstance->PortToUse = CmdLinePort;
+			ObjectInstance->PortsToUse.Add(0);
 		}
+
 		FParse::Value(FCommandLine::Get(), TEXT("rpcledgersize="), ObjectInstance->RequestLedgerCapacity);
 
 		// We always want the ListRegisteredRpcs route bound, no matter what.
@@ -285,25 +295,44 @@ void UExternalRpcRegistry::RegisterNewRoute(FExternalRouteInfo InRouteInfo, cons
 		return;
 	}
 
-	TSharedPtr<IHttpRouter> HttpRouter = FHttpServerModule::Get().GetHttpRouter(PortToUse);
-
-	if (RegisteredRoutes.Find(InRouteInfo.RouteName))
+	// create a route/route name for each port - if there are multiple ports, then their will be a "foo.portnum" named route for each port
+	// we also track a mapping from 'foo.portnum' to 'foo' for one of the routes, so when /listrpcs is called, we can report 'foo' as the rpc
+	// name that users of the api will use. 'foo.portnum' is not meant to be exposed externally
+	for (int PortToUse : PortsToUse)
 	{
-		if (!bOverrideIfBound)
+		TSharedPtr<IHttpRouter> HttpRouter = FHttpServerModule::Get().GetHttpRouter(PortToUse);
+
+		// for multiport listening, we need a route for each one
+		FName RoutePortName = InRouteInfo.RouteName;
+		if (PortsToUse.Num() > 1)
+		{
+			RoutePortName = *FString::Printf(TEXT("%s.%d"), *InRouteInfo.RouteName.ToString(), PortToUse);
+		}
+
+		if (RegisteredRoutes.Find(RoutePortName))
 		{
-			UE_LOG(LogExternalRpcRegistry, Error, TEXT("Failed to bind route with friendly key %s - a route at location %s already exists."), *InRouteInfo.RouteName.ToString(), *InRouteInfo.RoutePath.GetPath());
-			return;
+			if (!bOverrideIfBound)
+			{
+				UE_LOG(LogExternalRpcRegistry, Error, TEXT("Failed to bind route with friendly key %s - a route at location %s already exists."), *RoutePortName.ToString(), *InRouteInfo.RoutePath.GetPath());
+				return;
+			}
+			UE_LOG(LogExternalRpcRegistry, Log, TEXT("Overwriting route at friendly key %s - from %s to %s "), *RoutePortName.ToString(), *RegisteredRoutes[RoutePortName].Handle->Path, *InRouteInfo.RoutePath.GetPath());
+
+			HttpRouter->UnbindRoute(RegisteredRoutes[RoutePortName].Handle);
 		}
-		UE_LOG(LogExternalRpcRegistry, Log, TEXT("Overwriting route at friendly key %s - from %s to %s "), *InRouteInfo.RouteName.ToString(), *RegisteredRoutes[InRouteInfo.RouteName].Handle->Path, *InRouteInfo.RoutePath.GetPath());
-		HttpRouter->UnbindRoute(RegisteredRoutes[InRouteInfo.RouteName].Handle);
+
+		// keep a mapping of original route name to one of the new route names, for ListRPCs
+		MultiPortRouteMapping.Add(RoutePortName, InRouteInfo.RouteName);
+
+		FExternalRouteDesc RouteDesc;
+		RouteDesc.Handle = HttpRouter->BindRoute(InRouteInfo.RoutePath, InRouteInfo.RequestVerbs, Handler);
+		RouteDesc.InputContentType = InRouteInfo.InputContentType;
+		RouteDesc.ExpectedArguments = InRouteInfo.ExpectedArguments;
+		RouteDesc.RpcCategory = InRouteInfo.RpcCategory;
+		RegisteredRoutes.Add(RoutePortName, RouteDesc);
+		UE_LOG(LogExternalRpcRegistry, Log, TEXT("Route name %s was bound!"), *RoutePortName.ToString());
 	}
-	FExternalRouteDesc RouteDesc;
-	RouteDesc.Handle = HttpRouter->BindRoute(InRouteInfo.RoutePath, InRouteInfo.RequestVerbs, Handler);
-	RouteDesc.InputContentType = InRouteInfo.InputContentType;
-	RouteDesc.ExpectedArguments = InRouteInfo.ExpectedArguments;
-	RouteDesc.RpcCategory = InRouteInfo.RpcCategory;
-	RegisteredRoutes.Add(InRouteInfo.RouteName, RouteDesc);
-	UE_LOG(LogExternalRpcRegistry, Log, TEXT("Route name %s was bound!"), *InRouteInfo.RouteName.ToString());
+
 #endif
 }
 
@@ -324,11 +353,32 @@ void UExternalRpcRegistry::CleanUpRoute(FName RouteName, bool bFailIfUnbound /*
 #if WITH_RPC_REGISTRY
 	if (RegisteredRoutes.Find(RouteName))
 	{
-		TSharedPtr<IHttpRouter> HttpRouter = FHttpServerModule::Get().GetHttpRouter(PortToUse);
+		// get the port - if there are multiple ports, use the port from the end of the RouteName, otherwise, use the only PortToUse
+		int Port = 0;
+		if (PortsToUse.Num() > 1)
+		{
+			FString RouteNameString = RouteName.ToString();
+			int DotIndex;
+			if (!RouteNameString.FindLastChar('.', DotIndex))
+			{
+				UE_LOG(LogExternalRpcRegistry, Warning, TEXT("Expected a .portnum in the route name %s when using multport"), *RouteNameString);
+				check(!bFailIfUnbound);
+				return;
+			}
+			Port = FCString::Atoi(*RouteNameString.Mid(DotIndex + 1));
+		}
+		else
+		{
+			Port = PortsToUse[0];
+		}
+
+		// remove the mapping
+		MultiPortRouteMapping.Remove(RouteName);
+
+		TSharedPtr<IHttpRouter> HttpRouter = FHttpServerModule::Get().GetHttpRouter(Port);
 		HttpRouter->UnbindRoute(RegisteredRoutes[RouteName].Handle);
 		RegisteredRoutes.Remove(RouteName);
 		UE_LOG(LogExternalRpcRegistry, Log, TEXT("Route name %s was unbound!"), *RouteName.ToString());
-
 	}
 	else
 	{
@@ -343,14 +393,13 @@ bool UExternalRpcRegistry::HttpListOpenRoutes(const FHttpServerRequest& Request,
 {
 #if WITH_RPC_REGISTRY
 	FString ResponseStr;
-	TArray<FName> OutRouteKeys;
-	RegisteredRoutes.GetKeys(OutRouteKeys);
 	TSharedRef<TJsonWriter<>> JsonWriter = TJsonWriterFactory<>::Create(&ResponseStr);
 	JsonWriter->WriteArrayStart();
-	for (const FName& RouteKey : OutRouteKeys)
+	for (const TPair<FName, FName>& RouteKeyPair : MultiPortRouteMapping)
 	{
+		FName RouteKey = RouteKeyPair.Key;
 		JsonWriter->WriteObjectStart();
-		JsonWriter->WriteValue(TEXT("name"), RouteKey.ToString());
+		JsonWriter->WriteValue(TEXT("name"), RouteKeyPair.Value.ToString());
 		JsonWriter->WriteValue(TEXT("route"), RegisteredRoutes[RouteKey].Handle->Path);
 		JsonWriter->WriteValue(TEXT("verb"), GetHttpRouteVerbString(RegisteredRoutes[RouteKey].Handle->Verbs));
 		if (!RegisteredRoutes[RouteKey].InputContentType.IsEmpty())
@@ -465,7 +514,8 @@ bool UExternalRpcRegistry::HttpSwaggerUI(const FHttpServerRequest& Request, cons
 		Address = TEXT("127.0.0.1");
 	}
 
-	FStringFormatOrderedArguments Args{ Address,PortToUse };
+	// use primary port
+	FStringFormatOrderedArguments Args{ Address, PortsToUse[0]};
 	FString ListFormat = FString::Format(*ResponseStr, Args);
 	auto Response = FHttpServerResponse::Create(ListFormat, TEXT("text/html"));
 	OnComplete(MoveTemp(Response));
@@ -514,7 +564,8 @@ bool UExternalRpcRegistry::HttpListOASv3JSONRoutes(const FHttpServerRequest& Req
 	for (FString InternetAddr : Addresses)
 	{
 		JsonWriter->WriteObjectStart();
-		JsonWriter->WriteValue(TEXT("url"), FString::Printf(TEXT("http://%s:%d"), *InternetAddr, PortToUse));
+		// use primary port
+		JsonWriter->WriteValue(TEXT("url"), FString::Printf(TEXT("http://%s:%d"), *InternetAddr, PortsToUse[0]));
 		JsonWriter->WriteValue(TEXT("description"), TEXT("Default server access ip"));
 		JsonWriter->WriteObjectEnd();
 	}
diff --git a/Engine/Source/Runtime/ExternalRPCRegistry/Public/ExternalRpcRegistry.h b/Engine/Source/Runtime/ExternalRPCRegistry/Public/ExternalRpcRegistry.h
--- a/Engine/Source/Runtime/ExternalRPCRegistry/Public/ExternalRpcRegistry.h
+++ b/Engine/Source/Runtime/ExternalRPCRegistry/Public/ExternalRpcRegistry.h
@@ -144,8 +144,8 @@ struct FRpcConfig
 
 	TMap<FString, int> PortMapping;
 
-	// it is valid to have mappings, but no active port, in the case of a client with multiple endpoint hosts
-	int ActiveHostPort = 0;
+	// it is valid to have mappings, but no active ports, in the case of a client with multiple endpoint hosts
+	TArray<int> ActiveHostPorts;
 };
 
 
@@ -162,6 +162,8 @@ protected:
 	TMap<FName, FExternalRouteDesc> RegisteredRoutes;
 	TArray<FString> ActiveRpcCategories;
 	TArray<FRpcLedgerEntry> RequestLedger;
+	TMap<FName, FName> MultiPortRouteMapping;
+	TArray<int> PortsToUse;
 
 public:
 	static EXTERNALRPCREGISTRY_API UExternalRpcRegistry * GetInstance();
@@ -169,7 +171,6 @@ public:
 	static EXTERNALRPCREGISTRY_API const FRpcConfig& GetRpcConfig();
 	EXTERNALRPCREGISTRY_API ~UExternalRpcRegistry();
 
-	int PortToUse = 0;
 	int RequestLedgerCapacity = 10;
 
 	/**
