[Backout] - CL47849634
#rnx
Original CL Desc
-----------------------------------------------------------------
- Added support for multiple listening ports for Rpcs

diff --git a/Engine/Source/Runtime/ExternalRPCRegistry/Private/ExternalRpcRegistry.cpp b/Engine/Source/Runtime/ExternalRPCRegistry/Private/ExternalRpcRegistry.cpp
--- a/Engine/Source/Runtime/ExternalRPCRegistry/Private/ExternalRpcRegistry.cpp
+++ b/Engine/Source/Runtime/ExternalRPCRegistry/Private/ExternalRpcRegistry.cpp
@@ -86,22 +86,18 @@ const FRpcConfig& UExternalRpcRegistry::GetRpcConfig()
 	}
 	bHasInitialized = true;
 
+	FString ConfigName;
+	FParse::Value(FCommandLine::Get(), TEXT("-rpcconfig="), ConfigName);
 
-	// look for multiple -rpcport or -rpcconfig commandline options
-	int PortNum;
-	for (const TCHAR* Cursor = FCommandLine::Get(); (Cursor != nullptr) && FParse::Value(Cursor, TEXT("-rpcport="), PortNum, &Cursor); )
+	if (FParse::Value(FCommandLine::Get(), TEXT("-rpcport="), Config.ActiveHostPort))
 	{
-		Config.ActiveHostPorts.AddUnique(PortNum);
-		// set one as Default, just in case something asks
-		if (!Config.PortMapping.Contains(TEXT("Default")))
-		{
-			Config.PortMapping.Add(TEXT("Default"), PortNum);
-		}
-	}
+		checkf(ConfigName.Len() == 0, TEXT("It is not supported to run with -rpcport and -rpcconfig"));
 
-	// look for a -rpcconfig=FooRpc string 
-	FString ConfigName;
-	for (const TCHAR* Cursor = FCommandLine::Get(); (Cursor != nullptr) && FParse::Value(Cursor, TEXT("-rpcconfig="), ConfigName, false, &Cursor); )
+		// put it into Default, just in case
+		Config.PortMapping.Add(TEXT("Default"), Config.ActiveHostPort);
+	}
+	// if we don't have -rpcport, look for a -rpcconfig=FooRpc string 
+	else if (ConfigName.Len() > 0)
 	{
 		const FConfigSection* Section = GConfig->GetSection(*ConfigName, false, GEngineIni);
 		if (Section)
@@ -139,21 +135,16 @@ const FRpcConfig& UExternalRpcRegistry::GetRpcConfig()
 			}
 
 			// look up the active port, or use 0 for missing port (0 isn't a valid port)
-			PortNum = Config.PortMapping.FindRef(ActiveName);
+			Config.ActiveHostPort = Config.PortMapping.FindRef(ActiveName);
 
-			if (PortNum > 0)
+			// allow for multiple clients, so client nums >= 1 we append the num. client 0 is just called Client
+			if (ActiveName == TEXT("Client"))
 			{
-				// allow for multiple clients, so client nums >= 1 we append the num. client 0 is just called Client
-				if (ActiveName == TEXT("Client"))
+				int ClientNum = 0;
+				if (FParse::Value(FCommandLine::Get(), TEXT("-RpcClientNum="), ClientNum) && ClientNum >= 0 && ClientNum < 1000)
 				{
-					int ClientNum = 0;
-					if (FParse::Value(FCommandLine::Get(), TEXT("-RpcClientNum="), ClientNum) && ClientNum >= 0 && ClientNum < 1000)
-					{
-						PortNum += ClientNum;
-					}
+					Config.ActiveHostPort += ClientNum;
 				}
-
-				Config.ActiveHostPorts.AddUnique(PortNum);
 			}
 		}
 	}
@@ -164,7 +155,7 @@ const FRpcConfig& UExternalRpcRegistry::GetRpcConfig()
 bool UExternalRpcRegistry::IsEnabled()
 {
 #if WITH_RPC_REGISTRY
-	if (GetRpcConfig().ActiveHostPorts.Num() > 0)
+	if (GetRpcConfig().ActiveHostPort > 0)
 	{
 		return true;
 	}
@@ -194,12 +185,11 @@ UExternalRpcRegistry* UExternalRpcRegistry::GetInstance()
 		
 		// look for override from commandline (we want a value here even if UExternalRpcRegistry::IsEnabled() returns false, so the 
 		// routes below, and in other code, don't trigger an assert on registration)
-		ObjectInstance->PortsToUse = GetRpcConfig().ActiveHostPorts;
-		if (ObjectInstance->PortsToUse.Num() == 0)
+		int CmdLinePort = GetRpcConfig().ActiveHostPort;
+		if (CmdLinePort > 0)
 		{
-			ObjectInstance->PortsToUse.Add(0);
+			ObjectInstance->PortToUse = CmdLinePort;
 		}
-
 		FParse::Value(FCommandLine::Get(), TEXT("rpcledgersize="), ObjectInstance->RequestLedgerCapacity);
 
 		// We always want the ListRegisteredRpcs route bound, no matter what.
@@ -295,44 +285,25 @@ void UExternalRpcRegistry::RegisterNewRoute(FExternalRouteInfo InRouteInfo, cons
 		return;
 	}
 
-	// create a route/route name for each port - if there are multiple ports, then their will be a "foo.portnum" named route for each port
-	// we also track a mapping from 'foo.portnum' to 'foo' for one of the routes, so when /listrpcs is called, we can report 'foo' as the rpc
-	// name that users of the api will use. 'foo.portnum' is not meant to be exposed externally
-	for (int PortToUse : PortsToUse)
-	{
-		TSharedPtr<IHttpRouter> HttpRouter = FHttpServerModule::Get().GetHttpRouter(PortToUse);
-
-		// for multiport listening, we need a route for each one
-		FName RoutePortName = InRouteInfo.RouteName;
-		if (PortsToUse.Num() > 1)
-		{
-			RoutePortName = *FString::Printf(TEXT("%s.%d"), *InRouteInfo.RouteName.ToString(), PortToUse);
-		}
+	TSharedPtr<IHttpRouter> HttpRouter = FHttpServerModule::Get().GetHttpRouter(PortToUse);
 
-		if (RegisteredRoutes.Find(RoutePortName))
+	if (RegisteredRoutes.Find(InRouteInfo.RouteName))
+	{
+		if (!bOverrideIfBound)
 		{
-			if (!bOverrideIfBound)
-			{
-				UE_LOG(LogExternalRpcRegistry, Error, TEXT("Failed to bind route with friendly key %s - a route at location %s already exists."), *RoutePortName.ToString(), *InRouteInfo.RoutePath.GetPath());
-				return;
-			}
-			UE_LOG(LogExternalRpcRegistry, Log, TEXT("Overwriting route at friendly key %s - from %s to %s "), *RoutePortName.ToString(), *RegisteredRoutes[RoutePortName].Handle->Path, *InRouteInfo.RoutePath.GetPath());
-
-			HttpRouter->UnbindRoute(RegisteredRoutes[RoutePortName].Handle);
+			UE_LOG(LogExternalRpcRegistry, Error, TEXT("Failed to bind route with friendly key %s - a route at location %s already exists."), *InRouteInfo.RouteName.ToString(), *InRouteInfo.RoutePath.GetPath());
+			return;
 		}
-
-		// keep a mapping of original route name to one of the new route names, for ListRPCs
-		MultiPortRouteMapping.Add(RoutePortName, InRouteInfo.RouteName);
-
-		FExternalRouteDesc RouteDesc;
-		RouteDesc.Handle = HttpRouter->BindRoute(InRouteInfo.RoutePath, InRouteInfo.RequestVerbs, Handler);
-		RouteDesc.InputContentType = InRouteInfo.InputContentType;
-		RouteDesc.ExpectedArguments = InRouteInfo.ExpectedArguments;
-		RouteDesc.RpcCategory = InRouteInfo.RpcCategory;
-		RegisteredRoutes.Add(RoutePortName, RouteDesc);
-		UE_LOG(LogExternalRpcRegistry, Log, TEXT("Route name %s was bound!"), *RoutePortName.ToString());
+		UE_LOG(LogExternalRpcRegistry, Log, TEXT("Overwriting route at friendly key %s - from %s to %s "), *InRouteInfo.RouteName.ToString(), *RegisteredRoutes[InRouteInfo.RouteName].Handle->Path, *InRouteInfo.RoutePath.GetPath());
+		HttpRouter->UnbindRoute(RegisteredRoutes[InRouteInfo.RouteName].Handle);
 	}
-
+	FExternalRouteDesc RouteDesc;
+	RouteDesc.Handle = HttpRouter->BindRoute(InRouteInfo.RoutePath, InRouteInfo.RequestVerbs, Handler);
+	RouteDesc.InputContentType = InRouteInfo.InputContentType;
+	RouteDesc.ExpectedArguments = InRouteInfo.ExpectedArguments;
+	RouteDesc.RpcCategory = InRouteInfo.RpcCategory;
+	RegisteredRoutes.Add(InRouteInfo.RouteName, RouteDesc);
+	UE_LOG(LogExternalRpcRegistry, Log, TEXT("Route name %s was bound!"), *InRouteInfo.RouteName.ToString());
 #endif
 }
 
@@ -353,32 +324,11 @@ void UExternalRpcRegistry::CleanUpRoute(FName RouteName, bool bFailIfUnbound /*
 #if WITH_RPC_REGISTRY
 	if (RegisteredRoutes.Find(RouteName))
 	{
-		// get the port - if there are multiple ports, use the port from the end of the RouteName, otherwise, use the only PortToUse
-		int Port = 0;
-		if (PortsToUse.Num() > 1)
-		{
-			FString RouteNameString = RouteName.ToString();
-			int DotIndex;
-			if (!RouteNameString.FindLastChar('.', DotIndex))
-			{
-				UE_LOG(LogExternalRpcRegistry, Warning, TEXT("Expected a .portnum in the route name %s when using multport"), *RouteNameString);
-				check(!bFailIfUnbound);
-				return;
-			}
-			Port = FCString::Atoi(*RouteNameString.Mid(DotIndex + 1));
-		}
-		else
-		{
-			Port = PortsToUse[0];
-		}
-
-		// remove the mapping
-		MultiPortRouteMapping.Remove(RouteName);
-
-		TSharedPtr<IHttpRouter> HttpRouter = FHttpServerModule::Get().GetHttpRouter(Port);
+		TSharedPtr<IHttpRouter> HttpRouter = FHttpServerModule::Get().GetHttpRouter(PortToUse);
 		HttpRouter->UnbindRoute(RegisteredRoutes[RouteName].Handle);
 		RegisteredRoutes.Remove(RouteName);
 		UE_LOG(LogExternalRpcRegistry, Log, TEXT("Route name %s was unbound!"), *RouteName.ToString());
+
 	}
 	else
 	{
@@ -393,13 +343,14 @@ bool UExternalRpcRegistry::HttpListOpenRoutes(const FHttpServerRequest& Request,
 {
 #if WITH_RPC_REGISTRY
 	FString ResponseStr;
+	TArray<FName> OutRouteKeys;
+	RegisteredRoutes.GetKeys(OutRouteKeys);
 	TSharedRef<TJsonWriter<>> JsonWriter = TJsonWriterFactory<>::Create(&ResponseStr);
 	JsonWriter->WriteArrayStart();
-	for (const TPair<FName, FName>& RouteKeyPair : MultiPortRouteMapping)
+	for (const FName& RouteKey : OutRouteKeys)
 	{
-		FName RouteKey = RouteKeyPair.Key;
 		JsonWriter->WriteObjectStart();
-		JsonWriter->WriteValue(TEXT("name"), RouteKeyPair.Value.ToString());
+		JsonWriter->WriteValue(TEXT("name"), RouteKey.ToString());
 		JsonWriter->WriteValue(TEXT("route"), RegisteredRoutes[RouteKey].Handle->Path);
 		JsonWriter->WriteValue(TEXT("verb"), GetHttpRouteVerbString(RegisteredRoutes[RouteKey].Handle->Verbs));
 		if (!RegisteredRoutes[RouteKey].InputContentType.IsEmpty())
@@ -514,8 +465,7 @@ bool UExternalRpcRegistry::HttpSwaggerUI(const FHttpServerRequest& Request, cons
 		Address = TEXT("127.0.0.1");
 	}
 
-	// use primary port
-	FStringFormatOrderedArguments Args{ Address, PortsToUse[0]};
+	FStringFormatOrderedArguments Args{ Address,PortToUse };
 	FString ListFormat = FString::Format(*ResponseStr, Args);
 	auto Response = FHttpServerResponse::Create(ListFormat, TEXT("text/html"));
 	OnComplete(MoveTemp(Response));
@@ -564,8 +514,7 @@ bool UExternalRpcRegistry::HttpListOASv3JSONRoutes(const FHttpServerRequest& Req
 	for (FString InternetAddr : Addresses)
 	{
 		JsonWriter->WriteObjectStart();
-		// use primary port
-		JsonWriter->WriteValue(TEXT("url"), FString::Printf(TEXT("http://%s:%d"), *InternetAddr, PortsToUse[0]));
+		JsonWriter->WriteValue(TEXT("url"), FString::Printf(TEXT("http://%s:%d"), *InternetAddr, PortToUse));
 		JsonWriter->WriteValue(TEXT("description"), TEXT("Default server access ip"));
 		JsonWriter->WriteObjectEnd();
 	}
diff --git a/Engine/Source/Runtime/ExternalRPCRegistry/Public/ExternalRpcRegistry.h b/Engine/Source/Runtime/ExternalRPCRegistry/Public/ExternalRpcRegistry.h
--- a/Engine/Source/Runtime/ExternalRPCRegistry/Public/ExternalRpcRegistry.h
+++ b/Engine/Source/Runtime/ExternalRPCRegistry/Public/ExternalRpcRegistry.h
@@ -144,8 +144,8 @@ struct FRpcConfig
 
 	TMap<FString, int> PortMapping;
 
-	// it is valid to have mappings, but no active ports, in the case of a client with multiple endpoint hosts
-	TArray<int> ActiveHostPorts;
+	// it is valid to have mappings, but no active port, in the case of a client with multiple endpoint hosts
+	int ActiveHostPort = 0;
 };
 
 
@@ -162,8 +162,6 @@ protected:
 	TMap<FName, FExternalRouteDesc> RegisteredRoutes;
 	TArray<FString> ActiveRpcCategories;
 	TArray<FRpcLedgerEntry> RequestLedger;
-	TMap<FName, FName> MultiPortRouteMapping;
-	TArray<int> PortsToUse;
 
 public:
 	static EXTERNALRPCREGISTRY_API UExternalRpcRegistry * GetInstance();
@@ -171,6 +169,7 @@ public:
 	static EXTERNALRPCREGISTRY_API const FRpcConfig& GetRpcConfig();
 	EXTERNALRPCREGISTRY_API ~UExternalRpcRegistry();
 
+	int PortToUse = 0;
 	int RequestLedgerCapacity = 10;
 
 	/**
