Add Access-Control-Allow-Origin header to all RPC responses

#rnx

#Synced-CL 51159484

diff --git a/Engine/Source/Runtime/ExternalRPCRegistry/Private/ExternalRpcRegistry.cpp b/Engine/Source/Runtime/ExternalRPCRegistry/Private/ExternalRpcRegistry.cpp
--- a/Engine/Source/Runtime/ExternalRPCRegistry/Private/ExternalRpcRegistry.cpp
+++ b/Engine/Source/Runtime/ExternalRPCRegistry/Private/ExternalRpcRegistry.cpp
@@ -172,6 +172,30 @@ bool UExternalRpcRegistry::IsEnabled()
 	return false;
 }
 
+FHttpRequestHandler UExternalRpcRegistry::CreateRouteHandle(TDelegate<bool(const FHttpServerRequest& Request, const FHttpResultCallback& OnComplete)> InFunc)
+{
+	FHttpRequestHandler OutHandler = nullptr;
+
+#if WITH_RPC_REGISTRY
+	OutHandler = FHttpRequestHandler::CreateLambda([this, InFunc](const FHttpServerRequest& Request, const FHttpResultCallback& OnComplete)
+		{
+			UExternalRpcRegistry::GetInstance()->AddRequestToLedger(Request);
+
+			FHttpResultCallback WrappedOnComplete([&OnComplete](TUniquePtr<FHttpServerResponse> Response)
+				{
+					// Inject headers on every response here
+					Response->Headers.Add(TEXT("Access-Control-Allow-Origin"), { TEXT("*") });
+					OnComplete(MoveTemp(Response));
+				});
+
+			return InFunc.Execute(Request, WrappedOnComplete);
+
+		});
+#endif
+	return OutHandler;
+}
+
+
 UExternalRpcRegistry* UExternalRpcRegistry::GetInstance()
 {
 #if WITH_RPC_REGISTRY
@@ -206,21 +230,21 @@ UExternalRpcRegistry* UExternalRpcRegistry::GetInstance()
 
 		FHttpRequestHandler ListRoutesRequestHandler = FHttpRequestHandler::CreateUObject(ObjectInstance, &ThisClass::HttpListOpenRoutes);
 		ObjectInstance->RegisterNewRoute(TEXT("ListRegisteredRpcs"), FHttpPath("/listrpcs"), EHttpServerRequestVerbs::VERB_GET,
-			ListRoutesRequestHandler, true, true);
+			ObjectInstance->CreateRouteHandle(ListRoutesRequestHandler), true, true);
 
 		FHttpRequestHandler PrintLedgerRequestHandler = FHttpRequestHandler::CreateUObject(ObjectInstance, &ThisClass::HttpPrintRequestLedger);
 		// We always want the ListRegisteredRpcs route bound, no matter what.
 		ObjectInstance->RegisterNewRoute(TEXT("GetRequestHistory"), FHttpPath("/requesthistory"), EHttpServerRequestVerbs::VERB_GET,
-			PrintLedgerRequestHandler, true, true);
+			ObjectInstance->CreateRouteHandle(PrintLedgerRequestHandler), true, true);
 
 		FHttpRequestHandler ListOASv3RequestHandler = FHttpRequestHandler::CreateUObject(ObjectInstance, &ThisClass::HttpListOASv3JSONRoutes);
 		// /swagger.json escaped as %2e
 		ObjectInstance->RegisterNewRoute(TEXT("ListSwaggerJson"), FHttpPath("/swagger.json"), EHttpServerRequestVerbs::VERB_GET,
-			ListOASv3RequestHandler, true, true);
+			ObjectInstance->CreateRouteHandle(ListOASv3RequestHandler), true, true);
 
 		FHttpRequestHandler SwaggerUIHandler = FHttpRequestHandler::CreateUObject(ObjectInstance, &ThisClass::HttpSwaggerUI);
 		ObjectInstance->RegisterNewRoute(TEXT("SwaggerUIHTML"), FHttpPath("/swagger/index.html"), EHttpServerRequestVerbs::VERB_GET,
-			SwaggerUIHandler, true, true);
+			ObjectInstance->CreateRouteHandle(SwaggerUIHandler), true, true);
 
 		ObjectInstance->AddToRoot();
 	}
diff --git a/Engine/Source/Runtime/ExternalRPCRegistry/Public/ExternalRpcRegistry.h b/Engine/Source/Runtime/ExternalRPCRegistry/Public/ExternalRpcRegistry.h
--- a/Engine/Source/Runtime/ExternalRPCRegistry/Public/ExternalRpcRegistry.h
+++ b/Engine/Source/Runtime/ExternalRPCRegistry/Public/ExternalRpcRegistry.h
@@ -188,6 +188,8 @@ public:
 	 */
 	EXTERNALRPCREGISTRY_API bool GetRegisteredRoute(FName RouteName, FExternalRouteInfo& OutRouteInfo);
 
+	EXTERNALRPCREGISTRY_API FHttpRequestHandler CreateRouteHandle(TDelegate<bool(const FHttpServerRequest& Request, const FHttpResultCallback& OnComplete)> InFunc);
+
 	EXTERNALRPCREGISTRY_API void RegisterNewRoute(FExternalRouteInfo InRouteInfo, const FHttpRequestHandler& Handler, bool bOverrideIfBound = false);
 
 	/**
