RPC: Don't update IP's for Windows targets. This behavior is specifically for devkits. This fix addresses situations where windows servers start rejecting connections.

#Synced-CL 51088023

#Synced-CL 51088023

#Synced-CL 51088023

diff --git a/Engine/Source/Runtime/ExternalRPCRegistry/Private/ExternalRpcRegistrationComponent.cpp b/Engine/Source/Runtime/ExternalRPCRegistry/Private/ExternalRpcRegistrationComponent.cpp
--- a/Engine/Source/Runtime/ExternalRPCRegistry/Private/ExternalRpcRegistrationComponent.cpp
+++ b/Engine/Source/Runtime/ExternalRPCRegistry/Private/ExternalRpcRegistrationComponent.cpp
@@ -154,7 +154,7 @@ void UExternalRpcRegistrationComponent::RegisterHttpCallback(FName RouteName, co
 }
 bool UExternalRpcRegistrationComponent::HttpUpdateIpOnListener(FString TargetName, FString NewIP)
 {
-#if WITH_RPC_REGISTRY
+#if WITH_RPC_REGISTRY && !PLATFORM_WINDOWS
 	if (TargetName == TEXT("") || NewIP == TEXT(""))
 	{
 		// We don't have a listener.
