[RPC Framework] Adds ability to capture the in-engine execution time of any RPC.
You can request the execution time capture by supply by specifying "CaptureExecutionTime": "True" in the request json.
The response json will automatically be injected with "ExecutionTimeMs"
#Synced-CL 52050243

diff --git a/Engine/Source/Runtime/ExternalRPCRegistry/Private/ExternalRpcRegistrationComponent.cpp b/Engine/Source/Runtime/ExternalRPCRegistry/Private/ExternalRpcRegistrationComponent.cpp
--- a/Engine/Source/Runtime/ExternalRPCRegistry/Private/ExternalRpcRegistrationComponent.cpp
+++ b/Engine/Source/Runtime/ExternalRPCRegistry/Private/ExternalRpcRegistrationComponent.cpp
@@ -148,9 +148,32 @@ FHttpRequestHandler UExternalRpcRegistrationComponent::CreateRouteHandle(TDelega
 					return true;
 				}
 			}
+
 			UExternalRpcRegistry::GetInstance()->AddRequestToLedger(Request);
-			return InFunc.Execute(Request, OnComplete);
 
+			FHttpResultCallback CompletionCallback;
+			TSharedPtr<FJsonObject> RequestJson = GetJsonObjectFromRequestBody(Request.Body);
+			if (RequestJson.IsValid() && RequestJson->HasField(TEXT("CaptureExecutionTime")) && RequestJson->GetBoolField(TEXT("CaptureExecutionTime")))
+			{
+				FRpcTimingContext TimingContext;
+				TimingContext.StartTiming();
+
+				CompletionCallback = [this, OnComplete, TimingContext](TUniquePtr<FHttpServerResponse> Response) mutable
+				{
+					TimingContext.StopTiming();
+					// Inject timing into response JSON
+					AppendTimingContextToResponse(Response, TimingContext);
+
+					// Call original OnComplete to send response
+					OnComplete(MoveTemp(Response));
+				};
+			}
+			else
+			{
+				CompletionCallback = OnComplete;
+			}
+
+			return InFunc.Execute(Request, CompletionCallback);
 		});
 #endif
 	return OutHandler;
@@ -222,3 +245,45 @@ bool UExternalRpcRegistrationComponent::HttpUpdateIpOnListener(FString TargetNam
 return false;
 #endif
 }
+
+void UExternalRpcRegistrationComponent::AppendTimingContextToResponse(TUniquePtr<FHttpServerResponse>& InResponse, FRpcTimingContext& InTimingContext)
+{
+#if WITH_RPC_REGISTRY
+	if (!InResponse.IsValid() || InResponse->Body.IsEmpty())
+	{
+		return;
+	}
+
+	// Only append timing context to a json response
+	bool bIsJson = false;
+	if (const TArray<FString>* ContentTypeHeaders = InResponse->Headers.Find(TEXT("Content-Type")))
+	{
+		for (const FString& ContentType : *ContentTypeHeaders)
+		{
+			if (ContentType.Contains(TEXT("application/json"), ESearchCase::IgnoreCase))
+			{
+				bIsJson = true;
+				break;
+			}
+		}
+	}
+	if (!bIsJson)
+	{
+		return;
+	}
+
+	TSharedPtr<FJsonObject> ResponseJson = GetJsonObjectFromRequestBody(InResponse->Body);
+	ResponseJson->SetNumberField(TEXT("ExecutionTimeMs"), InTimingContext.ExecutionTimeMs);
+
+	// Re-serialize JSON
+	FString ResponseStr;
+	TSharedRef<TJsonWriter<>> Writer = TJsonWriterFactory<>::Create(&ResponseStr);
+	FJsonSerializer::Serialize(ResponseJson.ToSharedRef(), Writer);
+
+	// Re-write the response body
+	FTCHARToUTF8 ConvertToUtf8(*ResponseStr);
+	const uint8* ConvertToUtf8Bytes = (reinterpret_cast<const uint8*>(ConvertToUtf8.Get()));
+	InResponse->Body.Empty();
+	InResponse->Body.Append(ConvertToUtf8Bytes, ConvertToUtf8.Length());
+#endif
+}
diff --git a/Engine/Source/Runtime/ExternalRPCRegistry/Public/ExternalRpcRegistrationComponent.h b/Engine/Source/Runtime/ExternalRPCRegistry/Public/ExternalRpcRegistrationComponent.h
--- a/Engine/Source/Runtime/ExternalRPCRegistry/Public/ExternalRpcRegistrationComponent.h
+++ b/Engine/Source/Runtime/ExternalRPCRegistry/Public/ExternalRpcRegistrationComponent.h
@@ -119,7 +119,15 @@ public:
 	 */
 	UE_API TUniquePtr<FHttpServerResponse> CreateSimpleResponse(bool bInWasSuccessful, FString InValue = "", bool bInFatal = false);
 
+protected:
 
+	/**
+	 * Appends timing context information to a response 
+	 * 
+	 * @param  InResponse		The response to append execution time to
+	 * @param  InTimingContext  Handle to the timing context that was used to measure the execution time
+	 */
+	UE_API void AppendTimingContextToResponse(TUniquePtr<FHttpServerResponse>& InResponse, FRpcTimingContext& InTimingContext);
 };
 
 #undef UE_API
diff --git a/Engine/Source/Runtime/ExternalRPCRegistry/Public/ExternalRpcRegistry.h b/Engine/Source/Runtime/ExternalRPCRegistry/Public/ExternalRpcRegistry.h
--- a/Engine/Source/Runtime/ExternalRPCRegistry/Public/ExternalRpcRegistry.h
+++ b/Engine/Source/Runtime/ExternalRPCRegistry/Public/ExternalRpcRegistry.h
@@ -148,6 +148,37 @@ struct FRpcConfig
 	TArray<int> ActiveHostPorts;
 };
 
+/**
+ * Context struct used to measure execution time of RPCs
+ * Execution time can be returned in the response to measure actual runtime of the RPC minus the network traffic
+*/
+USTRUCT()
+struct FRpcTimingContext
+{
+	GENERATED_BODY()
+public:
+
+	double ExecutionTimeMs;
+
+	FRpcTimingContext() 
+		: ExecutionTimeMs(-1.0)
+	{ }
+
+	void StartTiming()
+	{
+		StartTime = FDateTime::Now();
+	}
+
+	double StopTiming()
+	{
+		ExecutionTimeMs = (FDateTime::Now() - StartTime).GetTotalMilliseconds();
+		return ExecutionTimeMs;
+	}
+
+private:
+	FDateTime StartTime;
+};
+
 
 /**
  * This class is designed to be a singleton that handles registry, maintenance, and cleanup of any REST endpoints exposed on the process 
